Define first TextureBuildWorkers for use with DDC as well as a DDC build worker base module used as an entrypoint for standalone worker binaries.

diff --git a/Engine/Source/Developer/TextureBuild/Private/TextureBuildFunction.cpp b/Engine/Source/Developer/TextureBuild/Private/TextureBuildFunction.cpp
new file mode 100644
--- /dev/null
+++ b/Engine/Source/Developer/TextureBuild/Private/TextureBuildFunction.cpp
@@ -0,0 +1,376 @@
+// Copyright Epic Games, Inc. All Rights Reserved.
+
+#include "TextureBuildFunction.h"
+
+#include "DerivedDataPayload.h"
+#include "Engine/TextureDefines.h"
+#include "IImageWrapper.h"
+#include "IImageWrapperModule.h"
+#include "ImageCore.h"
+#include "Interfaces/ITextureFormat.h"
+#include "Interfaces/ITextureFormatModule.h"
+#include "Memory/CompositeBuffer.h"
+#include "Memory/SharedBuffer.h"
+#include "Modules/ModuleManager.h"
+#include "PixelFormat.h"
+#include "Serialization/CompactBinary.h"
+#include "Serialization/FileRegions.h"
+#include "Serialization/MemoryWriter.h"
+#include "TextureCompressorModule.h"
+
+DEFINE_LOG_CATEGORY_STATIC(LogTextureBuildWorker, Log, All);
+
+#ifndef CASE_ENUM_TO_TEXT
+#define CASE_ENUM_TO_TEXT(txt) case txt: return TEXT(#txt);
+#endif
+
+static const TCHAR* GetPixelFormatString(EPixelFormat InPixelFormat)
+{
+	switch (InPixelFormat)
+	{
+		FOREACH_ENUM_EPIXELFORMAT(CASE_ENUM_TO_TEXT)
+	default:
+		return TEXT("PF_Unknown");
+	}
+}
+
+static void ReadCbField(FCbFieldView FieldView, FColor& OutColor)
+{
+	// Choosing the big endian ordering
+	FCbFieldViewIterator ColorCbArrayIt = FieldView.AsArrayView().CreateViewIterator();
+	OutColor.A = ColorCbArrayIt->AsUInt8(); ++ColorCbArrayIt;
+	OutColor.R = ColorCbArrayIt->AsUInt8(); ++ColorCbArrayIt;
+	OutColor.G = ColorCbArrayIt->AsUInt8(); ++ColorCbArrayIt;
+	OutColor.B = ColorCbArrayIt->AsUInt8(); ++ColorCbArrayIt;
+}
+
+static void ReadCbField(FCbFieldView FieldView, FVector4& OutVec4)
+{
+	FCbFieldViewIterator Vec4CbArrayIt = FieldView.AsArrayView().CreateViewIterator();
+	OutVec4.X = Vec4CbArrayIt->AsFloat(); ++Vec4CbArrayIt;
+	OutVec4.Y = Vec4CbArrayIt->AsFloat(); ++Vec4CbArrayIt;
+	OutVec4.Z = Vec4CbArrayIt->AsFloat(); ++Vec4CbArrayIt;
+	OutVec4.W = Vec4CbArrayIt->AsFloat(); ++Vec4CbArrayIt;
+}
+
+static void ReadCbField(FCbFieldView FieldView, FIntPoint& OutIntPoint)
+{
+	FCbFieldViewIterator IntPointCbArrayIt = FieldView.AsArrayView().CreateViewIterator();
+	OutIntPoint.X = IntPointCbArrayIt->AsInt32(); ++IntPointCbArrayIt;
+	OutIntPoint.Y = IntPointCbArrayIt->AsInt32(); ++IntPointCbArrayIt;
+}
+
+static void ReadBuildSettingsFromCompactBinary(const FCbObject& Object, FTextureBuildSettings& OutBuildSettings)
+{
+	FCbObjectView ColorAdjustmentCbObj = Object.FindView("ColorAdjustment").AsObjectView();
+	FColorAdjustmentParameters& ColorAdjustment = OutBuildSettings.ColorAdjustment;
+	ColorAdjustment.AdjustBrightness = ColorAdjustmentCbObj.FindView("AdjustBrightness").AsFloat();
+	ColorAdjustment.AdjustBrightnessCurve = ColorAdjustmentCbObj.FindView("AdjustBrightnessCurve").AsFloat();
+	ColorAdjustment.AdjustSaturation = ColorAdjustmentCbObj.FindView("AdjustSaturation").AsFloat();
+	ColorAdjustment.AdjustVibrance = ColorAdjustmentCbObj.FindView("AdjustVibrance").AsFloat();
+	ColorAdjustment.AdjustRGBCurve = ColorAdjustmentCbObj.FindView("AdjustRGBCurve").AsFloat();
+	ColorAdjustment.AdjustHue = ColorAdjustmentCbObj.FindView("AdjustHue").AsFloat();
+	ColorAdjustment.AdjustMinAlpha = ColorAdjustmentCbObj.FindView("AdjustMinAlpha").AsFloat();
+	ColorAdjustment.AdjustMaxAlpha = ColorAdjustmentCbObj.FindView("AdjustMaxAlpha").AsFloat();
+	ReadCbField(Object.FindView("AlphaCoverageThresholds"), OutBuildSettings.AlphaCoverageThresholds);
+	OutBuildSettings.MipSharpening = Object.FindView("MipSharpening").AsFloat();
+	OutBuildSettings.DiffuseConvolveMipLevel = Object.FindView("DiffuseConvolveMipLevel").AsUInt32();
+	OutBuildSettings.SharpenMipKernelSize = Object.FindView("SharpenMipKernelSize").AsUInt32();
+	OutBuildSettings.MaxTextureResolution = Object.FindView("MaxTextureResolution").AsUInt32();
+	OutBuildSettings.TextureFormatName = FName(Object.FindView("TextureFormatName").AsString());
+	OutBuildSettings.bHDRSource = Object.FindView("bHDRSource").AsBool();
+	OutBuildSettings.MipGenSettings = Object.FindView("MipGenSettings").AsUInt8();
+	OutBuildSettings.bCubemap = Object.FindView("bCubemap").AsBool();
+	OutBuildSettings.bTextureArray = Object.FindView("bTextureArray").AsBool();
+	OutBuildSettings.bVolume = Object.FindView("bVolume").AsBool();
+	OutBuildSettings.bLongLatSource = Object.FindView("bLongLatSource").AsBool();
+	OutBuildSettings.bSRGB = Object.FindView("bSRGB").AsBool();
+	OutBuildSettings.bUseLegacyGamma = Object.FindView("bUseLegacyGamma").AsBool();
+	OutBuildSettings.bPreserveBorder = Object.FindView("bPreserveBorder").AsBool();
+	OutBuildSettings.bForceAlphaChannel = Object.FindView("bForceAlphaChannel").AsBool();
+	OutBuildSettings.bDitherMipMapAlpha = Object.FindView("bDitherMipMapAlpha").AsBool();
+	OutBuildSettings.bComputeBokehAlpha = Object.FindView("bComputeBokehAlpha").AsBool();
+	OutBuildSettings.bReplicateRed = Object.FindView("bReplicateRed").AsBool();
+	OutBuildSettings.bReplicateAlpha = Object.FindView("bReplicateAlpha").AsBool();
+	OutBuildSettings.bDownsampleWithAverage = Object.FindView("bDownsampleWithAverage").AsBool();
+	OutBuildSettings.bSharpenWithoutColorShift = Object.FindView("bSharpenWithoutColorShift").AsBool();
+	OutBuildSettings.bBorderColorBlack = Object.FindView("bBorderColorBlack").AsBool();
+	OutBuildSettings.bFlipGreenChannel = Object.FindView("bFlipGreenChannel").AsBool();
+	OutBuildSettings.bApplyYCoCgBlockScale = Object.FindView("bApplyYCoCgBlockScale").AsBool();
+	OutBuildSettings.bApplyKernelToTopMip = Object.FindView("bApplyKernelToTopMip").AsBool();
+	OutBuildSettings.bRenormalizeTopMip = Object.FindView("bRenormalizeTopMip").AsBool();
+	OutBuildSettings.CompositeTextureMode = Object.FindView("CompositeTextureMode").AsUInt8();
+	OutBuildSettings.CompositePower = Object.FindView("CompositePower").AsFloat();
+	OutBuildSettings.LODBias = Object.FindView("LODBias").AsUInt32();
+	OutBuildSettings.LODBiasWithCinematicMips = Object.FindView("LODBiasWithCinematicMips").AsUInt32();
+	ReadCbField(Object.FindView("TopMipSize"), OutBuildSettings.TopMipSize);
+	OutBuildSettings.VolumeSizeZ = Object.FindView("VolumeSizeZ").AsInt32();
+	OutBuildSettings.ArraySlices = Object.FindView("ArraySlices").AsInt32();
+	OutBuildSettings.bStreamable = Object.FindView("bStreamable").AsBool();
+	OutBuildSettings.bVirtualStreamable = Object.FindView("bVirtualStreamable").AsBool();
+	OutBuildSettings.bChromaKeyTexture = Object.FindView("bChromaKeyTexture").AsBool();
+	OutBuildSettings.PowerOfTwoMode = Object.FindView("PowerOfTwoMode").AsUInt8();
+	ReadCbField(Object.FindView("PaddingColor"), OutBuildSettings.PaddingColor);
+	ReadCbField(Object.FindView("ChromaKeyColor"), OutBuildSettings.ChromaKeyColor);
+	OutBuildSettings.ChromaKeyThreshold = Object.FindView("ChromaKeyThreshold").AsFloat();
+	OutBuildSettings.CompressionQuality = Object.FindView("CompressionQuality").AsInt32();
+	OutBuildSettings.LossyCompressionAmount = Object.FindView("LossyCompressionAmount").AsInt32();
+	OutBuildSettings.Downscale = Object.FindView("Downscale").AsFloat();
+	OutBuildSettings.DownscaleOptions = Object.FindView("DownscaleOptions").AsUInt8();
+	OutBuildSettings.VirtualAddressingModeX = Object.FindView("VirtualAddressingModeX").AsInt32();
+	OutBuildSettings.VirtualAddressingModeY = Object.FindView("VirtualAddressingModeY").AsInt32();
+	OutBuildSettings.VirtualTextureTileSize = Object.FindView("VirtualTextureTileSize").AsInt32();
+	OutBuildSettings.VirtualTextureBorderSize = Object.FindView("VirtualTextureBorderSize").AsInt32();
+	OutBuildSettings.bVirtualTextureEnableCompressZlib = Object.FindView("bVirtualTextureEnableCompressZlib").AsBool();
+	OutBuildSettings.bVirtualTextureEnableCompressCrunch = Object.FindView("bVirtualTextureEnableCompressCrunch").AsBool();
+	OutBuildSettings.bHasEditorOnlyData = Object.FindView("bHasEditorOnlyData").AsBool();
+}
+
+static void ReadOutputSettingsFromCompactBinary(const FCbObject& Object, int32& NumInlineMips, FString& MipKeyPrefix)
+{
+	NumInlineMips = Object.FindView("NumInlineMips").AsInt32();
+	MipKeyPrefix = FString(Object.FindView("MipKeyPrefix").AsString());
+}
+
+static ERawImageFormat::Type ComputeRawImageFormat(ETextureSourceFormat SourceFormat)
+{
+	switch (SourceFormat)
+	{
+	case TSF_G8:		return ERawImageFormat::G8;
+	case TSF_G16:		return ERawImageFormat::G16;
+	case TSF_BGRA8:		return ERawImageFormat::BGRA8;
+	case TSF_BGRE8:		return ERawImageFormat::BGRE8;
+	case TSF_RGBA16:	return ERawImageFormat::RGBA16;
+	case TSF_RGBA16F:	return ERawImageFormat::RGBA16F;
+	default:
+		checkf(false, TEXT("Invalid source texture format encountered: %d."), SourceFormat);
+		return ERawImageFormat::G8;
+	}
+}
+
+static void ReadTextureSourceFromCompactBinary(const FCbObject& Object, UE::DerivedData::FBuildContext& Context, TArray<FImage>& OutMips)
+{
+	FAnsiStringView InputKey = Object.FindView("Input").AsString();
+	FString TempInputKey(InputKey);
+	FSharedBuffer InputBuffer = Context.GetInput(*TempInputKey);
+
+	ETextureSourceCompressionFormat CompressionFormat = (ETextureSourceCompressionFormat)Object.FindView("CompressionFormat").AsUInt8();
+	ETextureSourceFormat SourceFormat = (ETextureSourceFormat)Object.FindView("SourceFormat").AsUInt8();
+
+	ERawImageFormat::Type RawImageFormat = ComputeRawImageFormat(SourceFormat);
+
+	EGammaSpace GammaSpace = (EGammaSpace)Object.FindView("GammaSpace").AsUInt8();
+	int32 NumSlices = Object.FindView("NumSlices").AsInt32();
+	int32 SizeX = Object.FindView("SizeX").AsInt32();
+	int32 SizeY = Object.FindView("SizeY").AsInt32();
+	int32 MipSizeX = SizeX;
+	int32 MipSizeY = SizeY;
+	FCbArrayView MipsCbArrayView = Object.FindView("Mips").AsArrayView();
+	OutMips.Reserve(MipsCbArrayView.Num());
+	for (FCbFieldViewIterator MipsCbArrayIt = MipsCbArrayView.CreateViewIterator(); MipsCbArrayIt; ++MipsCbArrayIt)
+	{
+		FCbObjectView MipCbObjectView = MipsCbArrayIt->AsObjectView();
+		int64 MipOffset = MipCbObjectView.FindView("Offset").AsInt64();
+		int64 MipSize = MipCbObjectView.FindView("Size").AsInt64();
+
+		FImage* SourceMip = new(OutMips) FImage(
+			MipSizeX, MipSizeY,
+			NumSlices,
+			RawImageFormat,
+			GammaSpace
+		);
+
+		switch (CompressionFormat)
+		{
+		case TSCF_JPEG:
+		{
+			TSharedPtr<IImageWrapper> ImageWrapper = FModuleManager::LoadModuleChecked<IImageWrapperModule>(FName("ImageWrapper")).CreateImageWrapper(EImageFormat::JPEG);
+			ImageWrapper->SetCompressed((const uint8*)InputBuffer.GetData() + MipOffset, MipSize);
+			ImageWrapper->GetRaw(SourceFormat == TSF_G8 ? ERGBFormat::Gray : ERGBFormat::BGRA, 8, SourceMip->RawData);
+		}
+		break;
+		case TSCF_PNG:
+		{
+			TSharedPtr<IImageWrapper> ImageWrapper = FModuleManager::LoadModuleChecked<IImageWrapperModule>(FName("ImageWrapper")).CreateImageWrapper(EImageFormat::PNG);
+			ImageWrapper->SetCompressed((const uint8*)InputBuffer.GetData() + MipOffset, MipSize);
+			ERGBFormat RawFormat = (SourceFormat == TSF_G8 || SourceFormat == TSF_G16) ? ERGBFormat::Gray : ERGBFormat::RGBA;
+			ImageWrapper->GetRaw(RawFormat, (SourceFormat == TSF_G16 || SourceFormat == TSF_RGBA16) ? 16 : 8, SourceMip->RawData);
+		}
+		break;
+		default:
+			SourceMip->RawData.AddUninitialized(MipSize);
+			FMemory::Memcpy(
+				SourceMip->RawData.GetData(),
+				(const uint8*)InputBuffer.GetData() + MipOffset,
+				MipSize
+			);
+			break;
+		}
+
+		MipSizeX = FMath::Max(MipSizeX / 2, 1);
+		MipSizeY = FMath::Max(MipSizeY / 2, 1);
+	}
+
+}
+
+FTextureBuildFunction::FTextureBuildFunction()
+: Compressor(FModuleManager::LoadModuleChecked<ITextureCompressorModule>(TEXTURE_COMPRESSOR_MODULENAME))
+{
+}
+
+void FTextureBuildFunction::Build(UE::DerivedData::FBuildContext& Context) const
+{
+	const FCbObject BuildSettingsCbObj = Context.GetConstant(TEXT("TextureBuildSettings"));
+	FTextureBuildSettings BuildSettings;
+	ReadBuildSettingsFromCompactBinary(Context.GetConstant(TEXT("TextureBuildSettings")), BuildSettings);
+	int32 NumInlineMips;
+	FString MipKeyPrefix;
+	ReadOutputSettingsFromCompactBinary(Context.GetConstant(TEXT("TextureOutputSettings")), NumInlineMips, MipKeyPrefix);
+
+	TArray<FImage> SourceMips;
+	ReadTextureSourceFromCompactBinary(Context.GetConstant(TEXT("TextureSource")), Context, SourceMips);
+
+	TArray<FImage> AssociatedNormalSourceMips;
+	const FCbObject AssociatedNormalSourceMipsCbObj = Context.GetConstant(TEXT("CompositeTextureSource"));
+	if (AssociatedNormalSourceMipsCbObj)
+	{
+		ReadTextureSourceFromCompactBinary(AssociatedNormalSourceMipsCbObj, Context, AssociatedNormalSourceMips);
+	}
+
+	UE_LOG(LogTextureBuildWorker, Display, TEXT("Compressing %d source mip(s) (%dx%d) to %s..."), SourceMips.Num(), SourceMips[0].SizeX, SourceMips[0].SizeY, *BuildSettings.TextureFormatName.ToString());
+
+	TArray<FCompressedImage2D> CompressedMips;
+	uint32 NumMipsInTail;
+	uint32 ExtData;
+
+	Compressor.BuildTexture(SourceMips,
+		AssociatedNormalSourceMips,
+		BuildSettings,
+		CompressedMips,
+		NumMipsInTail,
+		ExtData);
+
+	// Write out streaming mips as payloads.
+	int32 MipCount = CompressedMips.Num();
+	const bool bForceAllMipsToBeInlined = BuildSettings.bCubemap || (BuildSettings.bVolume && !BuildSettings.bStreamable) || (BuildSettings.bTextureArray && !BuildSettings.bStreamable);
+	const int32 FirstInlineMip = bForceAllMipsToBeInlined ? 0 : FMath::Max(0, MipCount - FMath::Max(NumInlineMips, (int32)NumMipsInTail));
+	const int32 WritableMipCount = MipCount - ((NumMipsInTail > 0) ? ((int32)NumMipsInTail - 1) : 0);
+	for (int32 MipIndex = 0; MipIndex < WritableMipCount; ++MipIndex)
+	{
+		if (MipIndex < FirstInlineMip)
+		{
+			const FCompressedImage2D& CompressedMip = CompressedMips[MipIndex];
+			TStringBuilder<32> PayloadName;
+			PayloadName << "Mip" << MipIndex;
+
+			int32 MipSize = CompressedMip.RawData.Num();
+			FSharedBuffer MipHeader = FSharedBuffer::MakeView(&MipSize, sizeof(int32));
+			FCompositeBuffer CompositeMipBuffer(MipHeader, FSharedBuffer::MakeView(CompressedMips[MipIndex].RawData.GetData(), CompressedMips[MipIndex].RawData.Num()));
+
+			Context.AddPayload(UE::DerivedData::FPayloadId::FromName(*PayloadName), FCompressedBuffer::Compress(NAME_None, CompositeMipBuffer));
+		}
+	}
+
+	// Write out texture platform data and non-streaming mip tail as a single payload
+	// This is meant to match the behavior in SerializePlatformData in TextureDerivedData.cpp.
+	// TODO: The texture header/footer information here should not be serialized from the worker but could be left to the editor to fill in.
+	//		 Doing it here causes more code duplication than I would like.
+	{
+		bool bHasOptData = (NumMipsInTail != 0) || (ExtData != 0);
+		TArray<FSharedBuffer> OrderedBuffers;
+
+		TArray<uint8> TextureHeaderArray;
+		FMemoryWriter TextureHeaderWriter(TextureHeaderArray, /*bIsPersistent=*/ true);
+		TextureHeaderWriter << CompressedMips[0].SizeX;
+		TextureHeaderWriter << CompressedMips[0].SizeY;
+
+		const int32 NumSlices = BuildSettings.bCubemap ? 6 : (BuildSettings.bVolume || BuildSettings.bTextureArray) ? CompressedMips[0].SizeZ : 1;
+		static constexpr uint32 BitMask_CubeMap = 1u << 31u;
+		static constexpr uint32 BitMask_HasOptData = 1u << 30u;
+		static constexpr uint32 BitMask_NumSlices = BitMask_HasOptData - 1u;
+		uint32 PackedData = (NumSlices & BitMask_NumSlices) | (BuildSettings.bCubemap ? BitMask_CubeMap : 0) | (bHasOptData ? BitMask_HasOptData : 0);
+		TextureHeaderWriter << PackedData;
+
+		FString PixelFormatString(GetPixelFormatString((EPixelFormat)CompressedMips[0].PixelFormat));
+		TextureHeaderWriter << PixelFormatString;
+
+		if (bHasOptData)
+		{
+			TextureHeaderWriter << ExtData;
+			TextureHeaderWriter << NumMipsInTail;
+		}
+
+		TextureHeaderWriter << MipCount;
+
+		OrderedBuffers.Add(FSharedBuffer::MakeView(TextureHeaderArray.GetData(), TextureHeaderArray.Num()));
+
+		TArray<TArray<uint8>> MipHeaderArray;
+		MipHeaderArray.AddDefaulted(MipCount);
+
+		TArray<TArray<uint8>> MipFooterArray;
+		MipFooterArray.AddDefaulted(MipCount);
+
+		int64 CurrentOffset = TextureHeaderArray.Num();
+		for (int32 MipIndex = 0; MipIndex < MipCount; ++MipIndex)
+		{
+			bool bIsInlineMip = MipIndex >= FirstInlineMip;
+			FCompressedImage2D& CompressedMip = CompressedMips[MipIndex];
+
+			FMemoryWriter MipHeaderWriter(MipHeaderArray[MipIndex], /*bIsPersistent=*/ true);
+			bool bCooked = false;
+			MipHeaderWriter << bCooked;
+
+			uint32 BulkDataFlags = 0;
+			MipHeaderWriter << BulkDataFlags;
+
+			int32 BulkDataElementCount = bIsInlineMip ? CompressedMip.RawData.Num() : 0;
+			MipHeaderWriter << BulkDataElementCount;
+
+			int32 BulkDataBytesOnDisk = bIsInlineMip ? CompressedMip.RawData.Num() : 0;
+			MipHeaderWriter << BulkDataBytesOnDisk;
+
+			int64 BulkDataOffset = CurrentOffset + MipHeaderArray[MipIndex].Num() + sizeof(int64);
+			MipHeaderWriter << BulkDataOffset;
+
+			OrderedBuffers.Add(FSharedBuffer::MakeView(MipHeaderArray[MipIndex].GetData(), MipHeaderArray[MipIndex].Num()));
+			CurrentOffset += MipHeaderArray[MipIndex].Num();
+
+			if (bIsInlineMip)
+			{
+				OrderedBuffers.Add(FSharedBuffer::MakeView(CompressedMip.RawData.GetData(), CompressedMip.RawData.Num()));
+				CurrentOffset += CompressedMip.RawData.Num();
+			}
+
+
+			FMemoryWriter MipFooterWriter(MipFooterArray[MipIndex], /*bIsPersistent=*/ true);
+
+			MipFooterWriter << CompressedMip.SizeX;
+			MipFooterWriter << CompressedMip.SizeY;
+			MipFooterWriter << CompressedMip.SizeZ;
+
+			EFileRegionType FileRegion = FFileRegion::SelectType(EPixelFormat(CompressedMip.PixelFormat));
+			MipFooterWriter << FileRegion;
+
+			TStringBuilder<128> DerivedDataKey;
+			if (!bIsInlineMip)
+			{
+				// This has to match the behavior in GetTextureDerivedMipKey in TextureDerivedData.cpp
+				DerivedDataKey << MipKeyPrefix;
+				DerivedDataKey.Appendf(TEXT("_MIP%u_%dx%d"), MipIndex, CompressedMip.SizeX, CompressedMip.SizeY);
+			}
+			FString DerivedDataKeyFinal (DerivedDataKey.ToString());
+			MipFooterWriter << DerivedDataKeyFinal;
+
+			OrderedBuffers.Add(FSharedBuffer::MakeView(MipFooterArray[MipIndex].GetData(), MipFooterArray[MipIndex].Num()));
+			CurrentOffset += MipFooterArray[MipIndex].Num();
+		}
+
+		// Bool serialized as 32-bit int is the footer  
+		uint32 IsVirtual = 0;
+		OrderedBuffers.Add(FSharedBuffer::MakeView(&IsVirtual, sizeof(IsVirtual)));
+
+
+		FCompositeBuffer CompositeResult(OrderedBuffers);
+		Context.AddPayload(UE::DerivedData::FPayloadId::FromName("Texture"), FCompressedBuffer::Compress(NAME_None, CompositeResult));
+	}
+}
diff --git a/Engine/Source/Developer/TextureBuild/Public/TextureBuildFunction.h b/Engine/Source/Developer/TextureBuild/Public/TextureBuildFunction.h
new file mode 100644
--- /dev/null
+++ b/Engine/Source/Developer/TextureBuild/Public/TextureBuildFunction.h
@@ -0,0 +1,37 @@
+// Copyright Epic Games, Inc. All Rights Reserved.
+
+#pragma once
+
+#include "CoreTypes.h"
+
+#include "Containers/StringView.h"
+#include "DerivedDataBuildFunction.h"
+#include "Misc/Guid.h"
+
+class ITextureCompressorModule;
+
+class FTextureBuildFunction : public UE::DerivedData::IBuildFunction
+{
+public:
+	TEXTUREBUILD_API FTextureBuildFunction();
+
+	TEXTUREBUILD_API void SetName(FStringView InName)
+	{
+		Name = InName;
+	}
+	
+	TEXTUREBUILD_API virtual FStringView GetName() const override
+	{
+		return Name;
+	}
+
+	TEXTUREBUILD_API virtual FGuid GetVersion() const override
+	{
+		return FGuid();
+	}
+
+	TEXTUREBUILD_API virtual void Build(UE::DerivedData::FBuildContext& Context) const override;
+private:
+	ITextureCompressorModule& Compressor;
+	FString Name;
+};
diff --git a/Engine/Source/Developer/TextureBuild/TextureBuild.Build.cs b/Engine/Source/Developer/TextureBuild/TextureBuild.Build.cs
new file mode 100644
--- /dev/null
+++ b/Engine/Source/Developer/TextureBuild/TextureBuild.Build.cs
@@ -0,0 +1,25 @@
+// Copyright Epic Games, Inc. All Rights Reserved.
+
+using UnrealBuildTool;
+
+public class TextureBuild : ModuleRules
+{
+	public TextureBuild(ReadOnlyTargetRules Target) : base(Target)
+	{
+		PrivateIncludePathModuleNames.AddRange(new string[]
+		{
+			"CoreUObject",
+			"DerivedDataCache",
+			"Engine",
+		});
+
+		PrivateDependencyModuleNames.AddRange(
+			new string[] {
+				"Core",
+				"ImageCore",
+				"ImageWrapper",
+				"TextureCompressor",
+				"TextureFormat",
+			});
+	}
+}
diff --git a/Engine/Source/Programs/DerivedDataBuildWorker/DerivedDataBuildWorker.Build.cs b/Engine/Source/Programs/DerivedDataBuildWorker/DerivedDataBuildWorker.Build.cs
new file mode 100644
--- /dev/null
+++ b/Engine/Source/Programs/DerivedDataBuildWorker/DerivedDataBuildWorker.Build.cs
@@ -0,0 +1,79 @@
+// Copyright Epic Games, Inc. All Rights Reserved.
+
+using UnrealBuildTool;
+
+// Abstract base class for worker targets.  Not a valid target by itself, hence it is not put into a *.target.cs file.
+[SupportedPlatforms(UnrealPlatformClass.Desktop)]
+public abstract class DerivedDataBuildWorkerTarget : TargetRules
+{
+	public DerivedDataBuildWorkerTarget(TargetInfo Target) : base(Target)
+	{
+		Type = TargetType.Program;
+		LinkType = TargetLinkType.Monolithic;
+
+		SolutionDirectory = "Programs/BuildWorker";
+
+		bUseXGEController				= false;
+		bCompileFreeType				= false;
+		bLoggingToMemoryEnabled			= true;
+		bUseLoggingInShipping			= true;
+		bCompileWithAccessibilitySupport= false;
+		bWithServerCode					= false;
+		bCompileNavmeshClusterLinks		= false;
+		bCompileNavmeshSegmentLinks		= false;
+		bCompileRecast					= false;
+		bCompileICU 					= false;
+		bCompileChaos					= false;
+		bUseChaos						= false;
+		bWithLiveCoding					= false;
+		bBuildDeveloperTools			= false;
+		bUseMallocProfiler				= false;
+		bBuildWithEditorOnlyData		= false;
+		bCompileAgainstEngine			= false;
+		bCompileAgainstCoreUObject		= false;
+		bCompileAgainstApplicationCore	= false;
+		bUsesSlate = false;
+		bIsBuildingConsoleApplication	= true;
+		// TODO: I want to use static CRT in the future, but it causes link issues today likely due to 3rd party libraries
+		//bUseStaticCRT					= true;
+
+		WindowsPlatform.bUseBundledDbgHelp = false;
+		WindowsPlatform.bPixProfilingEnabled = false;
+		WindowsPlatform.bMergeIdenticalCOMDATs = true;
+		WindowsPlatform.bStripUnreferencedSymbols = true;
+
+		// Disable logging, as the workers are spawned often and logging will just slow them down
+		GlobalDefinitions.Add("ALLOW_LOG_FILE=0");
+		// Linking against wer.lib/wer.dll causes XGE to bail when the worker is run on a Windows 8 machine, so turn this off.
+		GlobalDefinitions.Add("ALLOW_WINDOWS_ERROR_REPORT_LIB=0");
+		// Disable external profiling in TextureBuilder to improve startup time
+		GlobalDefinitions.Add("UE_EXTERNAL_PROFILING_ENABLED=0");
+	}
+}
+
+public class DerivedDataBuildWorker : ModuleRules
+{
+	public DerivedDataBuildWorker(ReadOnlyTargetRules Target) : base(Target)
+	{
+		PublicIncludePaths.Add("Runtime/Launch/Public");
+		PublicIncludePaths.Add("Developer/DerivedDataCache/Public");
+
+		PrivateIncludePaths.Add("Runtime/Launch/Private");		// For LaunchEngineLoop.cpp include
+
+		PrivateIncludePathModuleNames.AddRange(
+			new string[] {
+				"DerivedDataCache",
+			});
+
+		PrivateDependencyModuleNames.AddRange(
+			new string[] {
+				"Core",
+				"Projects",
+			});
+
+		if (Target.bCompileAgainstApplicationCore)
+		{
+			PrivateDependencyModuleNames.Add("ApplicationCore");
+		}
+	}
+}
diff --git a/Engine/Source/Programs/DerivedDataBuildWorker/Private/DerivedDataBuildLoop.cpp b/Engine/Source/Programs/DerivedDataBuildWorker/Private/DerivedDataBuildLoop.cpp
new file mode 100644
--- /dev/null
+++ b/Engine/Source/Programs/DerivedDataBuildWorker/Private/DerivedDataBuildLoop.cpp
@@ -0,0 +1,314 @@
+// Copyright Epic Games, Inc. All Rights Reserved.
+
+#include "DerivedDataBuildLoop.h"
+
+#include "Compression/CompressedBuffer.h"
+#include "DerivedDataPayload.h"
+#include "Memory/SharedBuffer.h"
+#include "Misc/CommandLine.h"
+#include "Misc/Parse.h"
+#include "Misc/WildcardString.h"
+#include "Serialization/CompactBinaryValidation.h"
+#include "Serialization/CompactBinaryWriter.h"
+
+DEFINE_LOG_CATEGORY_STATIC(LogDerivedDataBuildLoop, Log, All);
+
+namespace UE::DerivedData
+{
+
+static void ParseCommandLine( const TCHAR* CmdLine, TArray<FString>& Tokens, TArray<FString>& Switches )
+{
+	FString NextToken;
+	while ( FParse::Token(CmdLine, NextToken, false) )
+	{
+		if ( **NextToken == TCHAR('-') )
+		{
+			new(Switches) FString(NextToken.Mid(1));
+		}
+		else
+		{
+			new(Tokens) FString(NextToken);
+		}
+	}
+}
+
+static FSharedBuffer LoadFile(const FString& Path)
+{
+	FSharedBuffer Buffer;
+
+	if (TUniquePtr<FArchive> Ar{IFileManager::Get().CreateFileReader(*Path, FILEREAD_Silent)})
+	{
+		const int64 TotalSize = Ar->TotalSize();
+		FUniqueBuffer MutableBuffer = FUniqueBuffer::Alloc(uint64(TotalSize));
+		Ar->Serialize(MutableBuffer.GetData(), TotalSize);
+		if (Ar->Close())
+		{
+			Buffer = MutableBuffer.MoveToShared();
+		}
+	}
+
+	return Buffer;
+}
+
+static FIoHash HashAndWriteToCompressedBufferFile(const FString& Directory, const void* InData, uint64 InDataSize)
+{
+	FIoHash DataHash = FIoHash::HashBuffer(InData, InDataSize);
+	TStringBuilder<128> DataHashStringBuilder;
+	DataHashStringBuilder << DataHash;
+
+	FCompressedBuffer CompressedBufferContents = FCompressedBuffer::Compress(NAME_None, FSharedBuffer::MakeView(InData, InDataSize));
+	if (TUniquePtr<FArchive> FileAr{IFileManager::Get().CreateFileWriter(*(Directory / *DataHashStringBuilder), FILEWRITE_NoReplaceExisting)})
+	{
+		*FileAr << CompressedBufferContents;
+	}
+	return DataHash;
+}
+
+class FWorkerBuildContext : public FBuildContext
+{
+public:
+	FWorkerBuildContext(const FBuildLoop::FBuildActionRecord& InBuildActionRecord)
+	: BuildActionRecord(InBuildActionRecord)
+	{
+		BuildWriter.BeginObject("BuildOutput");
+		BuildWriter.BeginArray("Payloads");
+	}
+
+	virtual ~FWorkerBuildContext()
+	{
+		BuildWriter.EndArray();
+
+		BuildWriter.EndObject();
+
+		if (TUniquePtr<FArchive> FileAr{IFileManager::Get().CreateFileWriter(*BuildActionRecord.OutputFilePath)})
+		{
+			BuildWriter.Save(*FileAr);
+		}
+	}
+
+	virtual FCbObject GetConstant(FStringView Key) const override
+	{
+		TStringConversion<FTCHARToUTF8_Convert> ConvertedKey(Key.GetData(), Key.Len());
+		return BuildActionRecord.BuildAction["Constants"].AsObject()[ConvertedKey].AsObject();
+	}
+
+	virtual FSharedBuffer GetInput(FStringView Key) const override
+	{
+		TStringConversion<FTCHARToUTF8_Convert> ConvertedKey(Key.GetData(), Key.Len());
+		FIoHash InputHash(BuildActionRecord.BuildAction["Inputs"].AsObject()[ConvertedKey].AsBinaryAttachment());
+		if (!InputHash.IsZero())
+		{
+			TStringBuilder<256> Path;
+			Path << BuildActionRecord.InputPath << TEXT('/') << InputHash;
+			return FCompressedBuffer::FromCompressed(LoadFile(*Path)).Decompress();
+		}
+		return FSharedBuffer();
+	}
+
+	virtual void AddPayload(const FPayload& Payload) override
+	{
+		AddPayload(Payload.GetId(), Payload.GetData());
+	}
+	virtual void AddPayload(const FPayloadId& Id, const FCompressedBuffer& Buffer)
+	{
+		BuildWriter.BeginObject();
+		BuildWriter.AddObjectId("Id", FCbObjectId(Id.GetView()));
+		BuildWriter.AddInteger("RawSize", Buffer.GetRawSize());
+
+		FIoHash IoHash(Buffer.GetRawHash());
+		TStringBuilder<128> DataHashStringBuilder;
+		DataHashStringBuilder << IoHash;
+		if (TUniquePtr<FArchive> FileAr{IFileManager::Get().CreateFileWriter(*(BuildActionRecord.OutputPath / *DataHashStringBuilder), FILEWRITE_NoReplaceExisting)})
+		{
+			*FileAr << const_cast<FCompressedBuffer&>(Buffer);
+		}
+		BuildWriter.AddBinaryAttachment("RawHash", IoHash);
+		BuildWriter.EndObject();
+	}
+	virtual void AddPayload(const FPayloadId& Id, const FSharedBuffer& Buffer) override
+	{
+		BuildWriter.BeginObject();
+		BuildWriter.AddObjectId("Id", FCbObjectId(Id.GetView()));
+		BuildWriter.AddInteger("RawSize", Buffer.GetSize());
+		FIoHash Hash = HashAndWriteToCompressedBufferFile(BuildActionRecord.OutputPath, Buffer.GetData(), Buffer.GetSize());
+		BuildWriter.AddBinaryAttachment("RawHash", Hash);
+		BuildWriter.EndObject();
+	}
+	virtual void AddPayload(const FPayloadId& Id, const FCbObject& Object) override
+	{
+		AddPayload(Id, Object.GetBuffer());
+	}
+
+	virtual void SetCachePolicy(ECachePolicy Policy) override
+	{
+		checkNoEntry();
+	}
+
+	virtual void BeginAsyncBuild() override
+	{
+		check(false);
+	}
+
+	virtual void EndAsyncBuild() override
+	{
+		check(false);
+	}
+
+private:
+	FCbWriter BuildWriter;
+	const FBuildLoop::FBuildActionRecord& BuildActionRecord;
+};
+
+FBuildLoop::FBuildActionRecord::FBuildActionRecord(const FString& InSourceFilePath, const FString& InCommonInputPath, const FString& InCommonOutputPath, FSharedBuffer&& InSharedBuffer)
+: SourceFilePath(InSourceFilePath)
+, OutputFilePath(FPaths::ChangeExtension(InSourceFilePath, TEXT("uddbo")))
+, InputPath(InCommonInputPath.IsEmpty() ? FPaths::GetPath(InSourceFilePath) / TEXT("Inputs") : InCommonInputPath)
+, OutputPath(InCommonOutputPath.IsEmpty() ? FPaths::GetPath(InSourceFilePath) / TEXT("Outputs") : InCommonOutputPath)
+, BuildAction(MoveTemp(InSharedBuffer))
+{
+}
+
+bool FBuildLoop::Init()
+{
+	TArray<FString> Tokens;
+	TArray<FString> Switches;
+	ParseCommandLine(FCommandLine::Get(), Tokens, Switches);
+
+	TArray<FString> InputDirectoryPaths;
+	TArray<FString> OutputDirectoryPaths;
+	TArray<FString> BuildFilePathPatterns;
+	for (int32 SwitchIdx = 0; SwitchIdx < Switches.Num(); SwitchIdx++)
+	{
+		const FStringView Switch = Switches[SwitchIdx];
+
+		auto GetSwitchValueElements = [&Switch](const FStringView SwitchKey) -> TArray<FString>
+		{
+			TArray<FString> ValueElements;
+			if (Switch.StartsWith(SwitchKey) == true)
+			{
+				FStringView ValuesList = Switch.Right(Switch.Len() - SwitchKey.Len());
+
+				// Allow support for -KEY=Value1+Value2+Value3 as well as -KEY=Value1 -KEY=Value2
+				for (int32 PlusIdx = ValuesList.Find(TEXT("+"), ESearchCase::CaseSensitive); PlusIdx != INDEX_NONE; PlusIdx = ValuesList.Find(TEXT("+"), ESearchCase::CaseSensitive))
+				{
+					ValueElements.Emplace(ValuesList.Left(PlusIdx));
+
+					ValuesList.RightInline(ValuesList.Len() - (PlusIdx + 1));
+				}
+				ValueElements.Emplace(ValuesList);
+			}
+			return ValueElements;
+		};
+
+		BuildFilePathPatterns += GetSwitchValueElements(TEXT("B="));
+		BuildFilePathPatterns += GetSwitchValueElements(TEXT("BUILD="));
+
+		InputDirectoryPaths += GetSwitchValueElements(TEXT("I="));
+		InputDirectoryPaths += GetSwitchValueElements(TEXT("INPUT="));
+
+		OutputDirectoryPaths += GetSwitchValueElements(TEXT("O="));
+		OutputDirectoryPaths += GetSwitchValueElements(TEXT("OUTPUT="));
+	}
+
+	switch (InputDirectoryPaths.Num())
+	{
+		case 1:
+		CommonInputPath = InputDirectoryPaths[0];
+		if (FPaths::IsRelative(CommonInputPath))
+		{
+			CommonInputPath = FPaths::Combine(FPaths::LaunchDir(), CommonInputPath);
+		}
+		case 0: // deliberate fallthrough
+		break;
+		default:
+		UE_LOG(LogDerivedDataBuildLoop, Error, TEXT("A maximum of one input directory can be specified, but '%d' were specified."), InputDirectoryPaths.Num());
+		return false;
+	}
+
+	switch (OutputDirectoryPaths.Num())
+	{
+		case 1:
+		CommonOutputPath = OutputDirectoryPaths[0];
+		if (FPaths::IsRelative(CommonOutputPath))
+		{
+			CommonOutputPath = FPaths::Combine(FPaths::LaunchDir(), CommonOutputPath);
+		}
+		case 0: // deliberate fallthrough
+		break;
+		default:
+		UE_LOG(LogDerivedDataBuildLoop, Error, TEXT("A maximum of one output directory can be specified, but '%d' were specified."), OutputDirectoryPaths.Num());
+		return false;
+	}
+
+	if (BuildFilePathPatterns.Num() == 0)
+	{
+		UE_LOG(LogDerivedDataBuildLoop, Error, TEXT("No build files specified on the commandline."));
+		return false;
+	}
+
+	for (const FString& BuildFilePathPattern : BuildFilePathPatterns)
+	{
+		TArray<FString> BuildActionFilePaths;
+		FWildcardString BuildFilePathWildcardString(BuildFilePathPattern);
+		if (BuildFilePathWildcardString.ContainsWildcards())
+		{
+			// TODO: Support wildcard matching - at least for filename pattern in a single directory.
+			UE_LOG(LogDerivedDataBuildLoop, Error, TEXT("Wildcards in the build commandline arguments are currently unsupported: '%s'"), *BuildFilePathWildcardString);
+			return false;
+		}
+		else
+		{
+			if (FPaths::IsRelative(BuildFilePathPattern))
+			{
+				BuildActionFilePaths.Add(FPaths::Combine(FPaths::LaunchDir(), BuildFilePathPattern));
+			}
+			else
+			{
+				BuildActionFilePaths.Add(BuildFilePathPattern);
+			}
+		}
+
+		for (const FString& BuildActionFilePath : BuildActionFilePaths)
+		{
+			UE_LOG(LogDerivedDataBuildLoop, Log, TEXT("Loading build file: '%s'"), *BuildActionFilePath);
+			FSharedBuffer BuildActionFileBuffer = LoadFile(BuildActionFilePath);
+
+			if (BuildActionFileBuffer.IsNull())
+			{
+				UE_LOG(LogDerivedDataBuildLoop, Error, TEXT("Missing build file: '%s'"), *BuildActionFilePath);
+				return false;
+			}
+
+			if (ValidateCompactBinaryRange(BuildActionFileBuffer, ECbValidateMode::Default) != ECbValidateError::None)
+			{
+				UE_LOG(LogDerivedDataBuildLoop, Error, TEXT("Invalid build file: '%s'"), *BuildActionFilePath);
+				return false;
+			}
+
+			BuildActionRecords.Emplace(BuildActionFilePath, CommonInputPath, CommonOutputPath, MoveTemp(BuildActionFileBuffer));
+		}
+	}
+
+	if (BuildActionRecords.Num() == 0)
+	{
+		UE_LOG(LogDerivedDataBuildLoop, Error, TEXT("No build actions to operate on."));
+		return false;
+	}
+
+	return true;
+}
+
+void FBuildLoop::PerformBuilds(const FBuildFunctionCallback& BuildFunctionCallback)
+{
+	for (const FBuildActionRecord& BuildActionRecord : BuildActionRecords)
+	{
+		FWorkerBuildContext Context(BuildActionRecord);
+		BuildFunctionCallback(FName(BuildActionRecord.BuildAction["Function"].AsObject()["Name"].AsString()), Context);
+	}
+}
+
+void FBuildLoop::Teardown()
+{
+}
+
+}
diff --git a/Engine/Source/Programs/DerivedDataBuildWorker/Private/DerivedDataBuildLoop.h b/Engine/Source/Programs/DerivedDataBuildWorker/Private/DerivedDataBuildLoop.h
new file mode 100644
--- /dev/null
+++ b/Engine/Source/Programs/DerivedDataBuildWorker/Private/DerivedDataBuildLoop.h
@@ -0,0 +1,44 @@
+// Copyright Epic Games, Inc. All Rights Reserved.
+
+#pragma once
+
+#include "CoreMinimal.h"
+
+#include "DerivedDataBuildFunction.h"
+#include "DerivedDataBuildWorkerInterface.h"
+#include "Serialization/CompactBinary.h"
+#include "Templates/Function.h"
+#include "UObject/NameTypes.h"
+
+namespace UE::DerivedData
+{
+
+class FBuildLoop
+{
+public:
+	bool Init();
+
+	using FBuildFunctionCallback = TFunctionRef<bool(FName, FBuildContext&)>;
+	void PerformBuilds(const FBuildFunctionCallback& BuildFunctionCallback);
+
+	void Teardown();
+private:
+	FString CommonInputPath;
+	FString CommonOutputPath;
+
+	friend class FWorkerBuildContext;
+
+	struct FBuildActionRecord
+	{
+		const FString SourceFilePath;
+		const FString OutputFilePath;
+		const FString InputPath;
+		const FString OutputPath;
+		const FCbObject BuildAction;
+
+		FBuildActionRecord(const FString& InSourceFilePath, const FString& InCommonInputPath, const FString& InCommonOutputPath, FSharedBuffer&& InSharedBuffer);
+	};
+	TArray<FBuildActionRecord> BuildActionRecords;
+};
+
+}
\ No newline at end of file
diff --git a/Engine/Source/Programs/DerivedDataBuildWorker/Private/DerivedDataBuildWorker.cpp b/Engine/Source/Programs/DerivedDataBuildWorker/Private/DerivedDataBuildWorker.cpp
new file mode 100644
--- /dev/null
+++ b/Engine/Source/Programs/DerivedDataBuildWorker/Private/DerivedDataBuildWorker.cpp
@@ -0,0 +1,73 @@
+// Copyright Epic Games, Inc. All Rights Reserved.
+
+#include "CoreMinimal.h"
+
+#include "Containers/Array.h"
+#include "CoreTypes.h"
+#include "DerivedDataBuildLoop.h"
+#include "DerivedDataBuildWorkerInterface.h"
+#include "Logging/LogMacros.h"
+#include "Misc/ScopeExit.h"
+#include "RequiredProgramMainCPPInclude.h"
+
+DEFINE_LOG_CATEGORY_STATIC(LogDerivedDataBuildWorker, Log, All);
+
+//////////////////////////////////////////////////////////////////////////
+
+extern void DerivedDataBuildWorkerInit();
+
+namespace UE::DerivedData
+{
+
+TMap<FName, IBuildFunction*> GBuildFunctions;
+
+void RegisterWorkerBuildFunction(IBuildFunction* BuildFunction)
+{
+	GBuildFunctions.FindOrAdd(FName(BuildFunction->GetName())) = BuildFunction;
+}
+
+}
+
+//////////////////////////////////////////////////////////////////////////
+
+INT32_MAIN_INT32_ARGC_TCHAR_ARGV()
+{
+	FTaskTagScope Scope(ETaskTag::EGameThread);
+	FString CmdLine = FCommandLine::BuildFromArgV(nullptr, ArgC, ArgV, nullptr);
+	GEngineLoop.PreInit(*CmdLine);
+
+	// Make sure the engine is properly cleaned up whenever we exit this function
+	ON_SCOPE_EXIT
+	{
+		FEngineLoop::AppPreExit();
+		FEngineLoop::AppExit();
+	};
+
+	UE::DerivedData::FBuildLoop BuildLoop;
+	if (!BuildLoop.Init())
+	{
+		return 1;
+	}
+
+	DerivedDataBuildWorkerInit();
+
+	BuildLoop.PerformBuilds([&] (FName FunctionName, UE::DerivedData::FBuildContext& BuildContext)
+		{
+			if (UE::DerivedData::IBuildFunction** FoundFunc = UE::DerivedData::GBuildFunctions.Find(FunctionName))
+			{
+				UE::DerivedData::IBuildFunction* BuildFunction = *FoundFunc;
+				UE_LOG(LogDerivedDataBuildWorker, Display, TEXT("Starting build function '%s'"), *FunctionName.ToString());
+				uint64 BuildStartTime = FPlatformTime::Cycles64();
+				BuildFunction->Build(BuildContext);
+				UE_LOG(LogDerivedDataBuildWorker, Display, TEXT("Completed in %fms"), FPlatformTime::ToMilliseconds64(FPlatformTime::Cycles64()-BuildStartTime));
+
+				return true;
+			}
+			UE_LOG(LogDerivedDataBuildWorker, Error, TEXT("Unknown build function: %s"), *FunctionName.ToString());
+			return false;
+		});
+
+	BuildLoop.Teardown();
+
+	return 0;
+}
diff --git a/Engine/Source/Programs/DerivedDataBuildWorker/Public/DerivedDataBuildWorkerInterface.h b/Engine/Source/Programs/DerivedDataBuildWorker/Public/DerivedDataBuildWorkerInterface.h
new file mode 100644
--- /dev/null
+++ b/Engine/Source/Programs/DerivedDataBuildWorker/Public/DerivedDataBuildWorkerInterface.h
@@ -0,0 +1,14 @@
+// Copyright Epic Games, Inc. All Rights Reserved.
+
+#pragma once
+
+#include "CoreMinimal.h"
+
+#include "DerivedDataBuildFunction.h"
+
+namespace UE::DerivedData
+{
+
+void RegisterWorkerBuildFunction(IBuildFunction* BuildFunction);
+
+} // UE::DerivedData
diff --git a/Engine/Source/Programs/UncompressedTextureBuildWorker/Private/UncompressedTextureBuildWorker.cpp b/Engine/Source/Programs/UncompressedTextureBuildWorker/Private/UncompressedTextureBuildWorker.cpp
new file mode 100644
--- /dev/null
+++ b/Engine/Source/Programs/UncompressedTextureBuildWorker/Private/UncompressedTextureBuildWorker.cpp
@@ -0,0 +1,19 @@
+// Copyright Epic Games, Inc. All Rights Reserved.
+
+#include "CoreMinimal.h"
+
+#include "DerivedDataBuildWorkerInterface.h"
+#include "LaunchEngineLoop.h"
+#include "Modules/ModuleManager.h"
+#include "TextureBuildFunction.h"
+
+IMPLEMENT_APPLICATION(UncompressedTextureBuildWorker, "UncompressedTextureBuildWorker");
+
+
+void DerivedDataBuildWorkerInit()
+{
+	static FTextureBuildFunction TextureBuildFunction;
+	TextureBuildFunction.SetName(TEXT("UncompressedTexture"));
+	UE::DerivedData::RegisterWorkerBuildFunction(&TextureBuildFunction);
+}
+
diff --git a/Engine/Source/Programs/UncompressedTextureBuildWorker/UncompressedTextureBuildWorker.Build.cs b/Engine/Source/Programs/UncompressedTextureBuildWorker/UncompressedTextureBuildWorker.Build.cs
new file mode 100644
--- /dev/null
+++ b/Engine/Source/Programs/UncompressedTextureBuildWorker/UncompressedTextureBuildWorker.Build.cs
@@ -0,0 +1,25 @@
+// Copyright Epic Games, Inc. All Rights Reserved.
+
+using UnrealBuildTool;
+
+public class UncompressedTextureBuildWorker : ModuleRules
+{
+	public UncompressedTextureBuildWorker(ReadOnlyTargetRules Target) : base(Target)
+	{
+		PrivateIncludePathModuleNames.AddRange(new string[]
+		{
+			"CoreUObject",
+			"Engine",
+			"DerivedDataCache",
+		});
+
+		PrivateDependencyModuleNames.AddRange(new string[]
+		{
+			"Core",
+			"DerivedDataBuildWorker",
+			"TextureBuild",
+			"TextureFormat",
+			"TextureFormatUncompressed",
+		});
+	}
+}
diff --git a/Engine/Source/Programs/UncompressedTextureBuildWorker/UncompressedTextureBuildWorker.Target.cs b/Engine/Source/Programs/UncompressedTextureBuildWorker/UncompressedTextureBuildWorker.Target.cs
new file mode 100644
--- /dev/null
+++ b/Engine/Source/Programs/UncompressedTextureBuildWorker/UncompressedTextureBuildWorker.Target.cs
@@ -0,0 +1,13 @@
+// Copyright Epic Games, Inc. All Rights Reserved.
+
+using UnrealBuildTool;
+using System.Collections.Generic;
+
+[SupportedPlatforms(UnrealPlatformClass.Desktop)]
+public class UncompressedTextureBuildWorkerTarget : DerivedDataBuildWorkerTarget
+{
+	public UncompressedTextureBuildWorkerTarget(TargetInfo Target) : base(Target)
+	{
+		LaunchModuleName = "UncompressedTextureBuildWorker";
+	}
+}
