DDC: Renamed GetInput+GetConstant to FindInput+FindConstant to clarify that they may return null

diff --git a/Engine/Source/Developer/DerivedDataCache/Private/DerivedDataBuildJobContext.cpp b/Engine/Source/Developer/DerivedDataCache/Private/DerivedDataBuildJobContext.cpp
--- a/Engine/Source/Developer/DerivedDataCache/Private/DerivedDataBuildJobContext.cpp
+++ b/Engine/Source/Developer/DerivedDataCache/Private/DerivedDataBuildJobContext.cpp
@@ -47,7 +47,7 @@ void FBuildJobContext::ResetInputs()
 	Inputs.Empty();
 }
 
-FCbObject FBuildJobContext::GetConstant(FStringView Key) const
+FCbObject FBuildJobContext::FindConstant(FStringView Key) const
 {
 	if (const FCbObject* Object = Constants.FindByHash(GetTypeHash(Key), Key))
 	{
@@ -56,7 +56,7 @@ FCbObject FBuildJobContext::GetConstant(FStringView Key) const
 	return FCbObject();
 }
 
-FSharedBuffer FBuildJobContext::GetInput(FStringView Key) const
+FSharedBuffer FBuildJobContext::FindInput(FStringView Key) const
 {
 	if (const FCompressedBuffer* Input = Inputs.FindByHash(GetTypeHash(Key), Key))
 	{
diff --git a/Engine/Source/Developer/DerivedDataCache/Private/DerivedDataBuildJobContext.h b/Engine/Source/Developer/DerivedDataCache/Private/DerivedDataBuildJobContext.h
--- a/Engine/Source/Developer/DerivedDataCache/Private/DerivedDataBuildJobContext.h
+++ b/Engine/Source/Developer/DerivedDataCache/Private/DerivedDataBuildJobContext.h
@@ -47,8 +47,8 @@ public:
 	void ResetInputs();
 
 private:
-	FCbObject GetConstant(FStringView Key) const final;
-	FSharedBuffer GetInput(FStringView Key) const final;
+	FCbObject FindConstant(FStringView Key) const final;
+	FSharedBuffer FindInput(FStringView Key) const final;
 
 	void AddPayload(const FPayload& Payload) final;
 	void AddPayload(const FPayloadId& Id, const FCompressedBuffer& Buffer) final;
diff --git a/Engine/Source/Developer/DerivedDataCache/Public/DerivedDataBuildFunction.h b/Engine/Source/Developer/DerivedDataCache/Public/DerivedDataBuildFunction.h
--- a/Engine/Source/Developer/DerivedDataCache/Public/DerivedDataBuildFunction.h
+++ b/Engine/Source/Developer/DerivedDataCache/Public/DerivedDataBuildFunction.h
@@ -82,10 +82,10 @@ public:
 	virtual ~FBuildContext() = default;
 
 	/** Returns the constant with the matching key, or an object with no fields if not found. */
-	virtual FCbObject GetConstant(FStringView Key) const = 0;
+	virtual FCbObject FindConstant(FStringView Key) const = 0;
 
 	/** Returns the input with the matching key, or a null buffer if not found. */
-	virtual FSharedBuffer GetInput(FStringView Key) const = 0;
+	virtual FSharedBuffer FindInput(FStringView Key) const = 0;
 
 	/** Adds a payload to the build output. Must have a non-null buffer and a unique ID. */
 	virtual void AddPayload(const FPayload& Payload) = 0;
@@ -125,7 +125,7 @@ public:
 	virtual ~FBuildConfigContext() = default;
 
 	/** Returns the constant with the matching key, or an object with no fields if not found. */
-	virtual FCbObject GetConstant(FStringView Key) const = 0;
+	virtual FCbObject FindConstant(FStringView Key) const = 0;
 
 	/** Overrides the cache bucket used when reading or writing this build in the cache. */
 	virtual void SetCacheBucket(FCacheBucket Bucket) = 0;
diff --git a/Engine/Source/Developer/TextureBuild/Private/TextureBuildFunction.cpp b/Engine/Source/Developer/TextureBuild/Private/TextureBuildFunction.cpp
--- a/Engine/Source/Developer/TextureBuild/Private/TextureBuildFunction.cpp
+++ b/Engine/Source/Developer/TextureBuild/Private/TextureBuildFunction.cpp
@@ -153,7 +153,7 @@ static ERawImageFormat::Type ComputeRawImageFormat(ETextureSourceFormat SourceFo
 static void ReadTextureSourceFromCompactBinary(const FCbObject& Object, UE::DerivedData::FBuildContext& Context, TArray<FImage>& OutMips)
 {
 	FAnsiStringView InputKey = Object.FindView("Input").AsString();
-	FSharedBuffer InputBuffer = Context.GetInput(FUTF8ToTCHAR(InputKey));
+	FSharedBuffer InputBuffer = Context.FindInput(FUTF8ToTCHAR(InputKey));
 
 	ETextureSourceCompressionFormat CompressionFormat = (ETextureSourceCompressionFormat)Object.FindView("CompressionFormat").AsUInt8();
 	ETextureSourceFormat SourceFormat = (ETextureSourceFormat)Object.FindView("SourceFormat").AsUInt8();
@@ -236,18 +236,18 @@ void FTextureBuildFunction::Build(UE::DerivedData::FBuildContext& Context) const
 {
 	ITextureCompressorModule& Compressor = FModuleManager::LoadModuleChecked<ITextureCompressorModule>(TEXTURE_COMPRESSOR_MODULENAME);
 
-	const FCbObject BuildSettingsCbObj = Context.GetConstant(TEXT("TextureBuildSettings"));
+	const FCbObject BuildSettingsCbObj = Context.FindConstant(TEXT("TextureBuildSettings"));
 	FTextureBuildSettings BuildSettings;
-	ReadBuildSettingsFromCompactBinary(Context.GetConstant(TEXT("TextureBuildSettings")), BuildSettings);
+	ReadBuildSettingsFromCompactBinary(Context.FindConstant(TEXT("TextureBuildSettings")), BuildSettings);
 	int32 NumInlineMips;
 	FString MipKeyPrefix;
-	ReadOutputSettingsFromCompactBinary(Context.GetConstant(TEXT("TextureOutputSettings")), NumInlineMips, MipKeyPrefix);
+	ReadOutputSettingsFromCompactBinary(Context.FindConstant(TEXT("TextureOutputSettings")), NumInlineMips, MipKeyPrefix);
 
 	TArray<FImage> SourceMips;
-	ReadTextureSourceFromCompactBinary(Context.GetConstant(TEXT("TextureSource")), Context, SourceMips);
+	ReadTextureSourceFromCompactBinary(Context.FindConstant(TEXT("TextureSource")), Context, SourceMips);
 
 	TArray<FImage> AssociatedNormalSourceMips;
-	const FCbObject AssociatedNormalSourceMipsCbObj = Context.GetConstant(TEXT("CompositeTextureSource"));
+	const FCbObject AssociatedNormalSourceMipsCbObj = Context.FindConstant(TEXT("CompositeTextureSource"));
 	if (AssociatedNormalSourceMipsCbObj)
 	{
 		ReadTextureSourceFromCompactBinary(AssociatedNormalSourceMipsCbObj, Context, AssociatedNormalSourceMips);
