Texture: Moved textures to the Texture cache bucket

#rnx

diff --git a/Engine/Source/Developer/TextureBuild/Private/TextureBuildFunction.cpp b/Engine/Source/Developer/TextureBuild/Private/TextureBuildFunction.cpp
--- a/Engine/Source/Developer/TextureBuild/Private/TextureBuildFunction.cpp
+++ b/Engine/Source/Developer/TextureBuild/Private/TextureBuildFunction.cpp
@@ -2,6 +2,7 @@
 
 #include "TextureBuildFunction.h"
 
+#include "DerivedDataCache.h"
 #include "DerivedDataPayload.h"
 #include "Engine/TextureDefines.h"
 #include "IImageWrapper.h"
@@ -232,6 +233,11 @@ static void ReadTextureSourceFromCompactBinary(const FCbObject& Object, UE::Deri
 
 }
 
+void FTextureBuildFunction::Configure(UE::DerivedData::FBuildConfigContext& Context) const
+{
+	Context.SetCacheBucket(Context.CreateCacheBucket(TEXT("Texture")));
+}
+
 void FTextureBuildFunction::Build(UE::DerivedData::FBuildContext& Context) const
 {
 	ITextureCompressorModule& Compressor = FModuleManager::LoadModuleChecked<ITextureCompressorModule>(TEXTURE_COMPRESSOR_MODULENAME);
diff --git a/Engine/Source/Developer/TextureBuild/Public/TextureBuildFunction.h b/Engine/Source/Developer/TextureBuild/Public/TextureBuildFunction.h
--- a/Engine/Source/Developer/TextureBuild/Public/TextureBuildFunction.h
+++ b/Engine/Source/Developer/TextureBuild/Public/TextureBuildFunction.h
@@ -9,6 +9,6 @@
 class FTextureBuildFunction : public UE::DerivedData::IBuildFunction
 {
 public:
+	TEXTUREBUILD_API virtual void Configure(UE::DerivedData::FBuildConfigContext& Context) const override;
 	TEXTUREBUILD_API virtual void Build(UE::DerivedData::FBuildContext& Context) const override;
-private:
 };
