Texture: Fixed the formatting of an error message

#rnx

diff --git a/Engine/Source/Developer/TextureBuild/Private/TextureBuildFunction.cpp b/Engine/Source/Developer/TextureBuild/Private/TextureBuildFunction.cpp
--- a/Engine/Source/Developer/TextureBuild/Private/TextureBuildFunction.cpp
+++ b/Engine/Source/Developer/TextureBuild/Private/TextureBuildFunction.cpp
@@ -172,7 +172,7 @@ static bool TryReadTextureSourceFromCompactBinary(FCbFieldView Source, UE::Deriv
 	FSharedBuffer InputBuffer = Context.FindInput(FUTF8ToTCHAR(Source.GetName()));
 	if (!InputBuffer)
 	{
-		UE_LOG(LogTextureBuildFunction, Error, TEXT("Missing input %.*s."), Source.GetName().Len(), Source.GetName().GetData());
+		UE_LOG(LogTextureBuildFunction, Error, TEXT("Missing input '%s'."), *WriteToString<64>(FUTF8ToTCHAR(Source.GetName())));
 		return false;
 	}
 
