FTextureBuildFunction::Build: Fix crash if compressor build fails.
#rnx

diff --git a/Engine/Source/Developer/TextureBuild/Private/TextureBuildFunction.cpp b/Engine/Source/Developer/TextureBuild/Private/TextureBuildFunction.cpp
--- a/Engine/Source/Developer/TextureBuild/Private/TextureBuildFunction.cpp
+++ b/Engine/Source/Developer/TextureBuild/Private/TextureBuildFunction.cpp
@@ -306,13 +306,18 @@ void FTextureBuildFunction::Build(UE::DerivedData::FBuildContext& Context) const
 	uint32 NumMipsInTail;
 	uint32 ExtData;
 
-	FModuleManager::LoadModuleChecked<ITextureCompressorModule>(TEXTURE_COMPRESSOR_MODULENAME).BuildTexture(
+	bool bBuildSucceeded = FModuleManager::LoadModuleChecked<ITextureCompressorModule>(TEXTURE_COMPRESSOR_MODULENAME).BuildTexture(
 		SourceMips,
 		AssociatedNormalSourceMips,
 		BuildSettings,
 		CompressedMips,
 		NumMipsInTail,
 		ExtData);
+	if (!bBuildSucceeded)
+	{
+		return;
+	}
+	check(CompressedMips.Num() > 0);
 
 	// Write out streaming mips as payloads.
 	int32 MipCount = CompressedMips.Num();
