Texture: Use GetModuleChecked instead of LoadModuleChecked since it can be called from any thread

#rnx

diff --git a/Engine/Source/Developer/TextureBuild/Private/TextureBuildFunction.cpp b/Engine/Source/Developer/TextureBuild/Private/TextureBuildFunction.cpp
--- a/Engine/Source/Developer/TextureBuild/Private/TextureBuildFunction.cpp
+++ b/Engine/Source/Developer/TextureBuild/Private/TextureBuildFunction.cpp
@@ -202,14 +202,14 @@ static bool TryReadTextureSourceFromCompactBinary(FCbFieldView Source, UE::Deriv
 		{
 		case TSCF_JPEG:
 		{
-			TSharedPtr<IImageWrapper> ImageWrapper = FModuleManager::LoadModuleChecked<IImageWrapperModule>(FName("ImageWrapper")).CreateImageWrapper(EImageFormat::JPEG);
+			TSharedPtr<IImageWrapper> ImageWrapper = FModuleManager::GetModuleChecked<IImageWrapperModule>(FName("ImageWrapper")).CreateImageWrapper(EImageFormat::JPEG);
 			ImageWrapper->SetCompressed((const uint8*)InputBuffer.GetData(), InputBuffer.GetSize());
 			ImageWrapper->GetRaw(SourceFormat == TSF_G8 ? ERGBFormat::Gray : ERGBFormat::BGRA, 8, IntermediateDecompressedData);
 		}
 		break;
 		case TSCF_PNG:
 		{
-			TSharedPtr<IImageWrapper> ImageWrapper = FModuleManager::LoadModuleChecked<IImageWrapperModule>(FName("ImageWrapper")).CreateImageWrapper(EImageFormat::PNG);
+			TSharedPtr<IImageWrapper> ImageWrapper = FModuleManager::GetModuleChecked<IImageWrapperModule>(FName("ImageWrapper")).CreateImageWrapper(EImageFormat::PNG);
 			ImageWrapper->SetCompressed((const uint8*)InputBuffer.GetData(), InputBuffer.GetSize());
 			ERGBFormat RawFormat = (SourceFormat == TSF_G8 || SourceFormat == TSF_G16) ? ERGBFormat::Gray : ERGBFormat::RGBA;
 			ImageWrapper->GetRaw(RawFormat, (SourceFormat == TSF_G16 || SourceFormat == TSF_RGBA16) ? 16 : 8, IntermediateDecompressedData);
@@ -320,7 +320,7 @@ void FTextureBuildFunction::Build(UE::DerivedData::FBuildContext& Context) const
 	uint32 NumMipsInTail;
 	uint32 ExtData;
 
-	bool bBuildSucceeded = FModuleManager::LoadModuleChecked<ITextureCompressorModule>(TEXTURE_COMPRESSOR_MODULENAME).BuildTexture(
+	bool bBuildSucceeded = FModuleManager::GetModuleChecked<ITextureCompressorModule>(TEXTURE_COMPRESSOR_MODULENAME).BuildTexture(
 		SourceMips,
 		AssociatedNormalSourceMips,
 		BuildSettings,
