Texture: Allow DDC to select the compression method for texture build output and remove the mip size from serialized mips

#rnx

diff --git a/Engine/Source/Developer/TextureBuild/Private/TextureBuildFunction.cpp b/Engine/Source/Developer/TextureBuild/Private/TextureBuildFunction.cpp
--- a/Engine/Source/Developer/TextureBuild/Private/TextureBuildFunction.cpp
+++ b/Engine/Source/Developer/TextureBuild/Private/TextureBuildFunction.cpp
@@ -26,7 +26,7 @@ DEFINE_LOG_CATEGORY_STATIC(LogTextureBuildFunction, Log, All);
 // MUST have a corresponding change to this version. Individual texture formats have a version to
 // change that is specific to the format. A merge conflict affecting the version MUST be resolved
 // by generating a new version.
-static const FGuid TextureDerivedDataVersion(TEXT("a24fc8e0-42cb-49e8-bcd2-c8c4aa064bbd"));
+static const FGuid TextureDerivedDataVersion(TEXT("d565a540-b94c-4c4a-9852-16e343aba4b5"));
 
 #ifndef CASE_ENUM_TO_TEXT
 #define CASE_ENUM_TO_TEXT(txt) case txt: return TEXT(#txt);
@@ -333,26 +333,9 @@ void FTextureBuildFunction::Build(UE::DerivedData::FBuildContext& Context) const
 	}
 	check(CompressedMips.Num() > 0);
 
-	// Write out streaming mips as payloads.
 	int32 MipCount = CompressedMips.Num();
 	const bool bForceAllMipsToBeInlined = BuildSettings.bCubemap || (BuildSettings.bVolume && !BuildSettings.bStreamable) || (BuildSettings.bTextureArray && !BuildSettings.bStreamable);
 	const int32 FirstInlineMip = bForceAllMipsToBeInlined ? 0 : FMath::Max(0, MipCount - FMath::Max(NumInlineMips, (int32)NumMipsInTail));
-	const int32 WritableMipCount = MipCount - ((NumMipsInTail > 0) ? ((int32)NumMipsInTail - 1) : 0);
-	for (int32 MipIndex = 0; MipIndex < WritableMipCount; ++MipIndex)
-	{
-		if (MipIndex < FirstInlineMip)
-		{
-			const FCompressedImage2D& CompressedMip = CompressedMips[MipIndex];
-			TStringBuilder<32> PayloadName;
-			PayloadName << "Mip" << MipIndex;
-
-			int32 MipSize = CompressedMip.RawData.Num();
-			FSharedBuffer MipHeader = FSharedBuffer::MakeView(&MipSize, sizeof(int32));
-			FCompositeBuffer CompositeMipBuffer(MipHeader, FSharedBuffer::MakeView(CompressedMips[MipIndex].RawData.GetData(), CompressedMips[MipIndex].RawData.Num()));
-
-			Context.AddPayload(UE::DerivedData::FPayloadId::FromName(*PayloadName), FCompressedBuffer::Compress(NAME_Default, CompositeMipBuffer));
-		}
-	}
 
 	// Write out texture platform data and non-streaming mip tail as a single payload
 	// This is meant to match the behavior in SerializePlatformData in TextureDerivedData.cpp.
@@ -387,19 +370,14 @@ void FTextureBuildFunction::Build(UE::DerivedData::FBuildContext& Context) const
 
 		OrderedBuffers.Add(FSharedBuffer::MakeView(TextureHeaderArray.GetData(), TextureHeaderArray.Num()));
 
-		TArray<TArray<uint8>> MipHeaderArray;
-		MipHeaderArray.AddDefaulted(MipCount);
-
-		TArray<TArray<uint8>> MipFooterArray;
-		MipFooterArray.AddDefaulted(MipCount);
-
 		int64 CurrentOffset = TextureHeaderArray.Num();
 		for (int32 MipIndex = 0; MipIndex < MipCount; ++MipIndex)
 		{
 			bool bIsInlineMip = MipIndex >= FirstInlineMip;
 			FCompressedImage2D& CompressedMip = CompressedMips[MipIndex];
 
-			FMemoryWriter MipHeaderWriter(MipHeaderArray[MipIndex], /*bIsPersistent=*/ true);
+			TArray<uint8> MipHeader;
+			FMemoryWriter MipHeaderWriter(MipHeader, /*bIsPersistent=*/ true);
 			bool bCooked = false;
 			MipHeaderWriter << bCooked;
 
@@ -412,11 +390,11 @@ void FTextureBuildFunction::Build(UE::DerivedData::FBuildContext& Context) const
 			int32 BulkDataBytesOnDisk = bIsInlineMip ? CompressedMip.RawData.Num() : 0;
 			MipHeaderWriter << BulkDataBytesOnDisk;
 
-			int64 BulkDataOffset = CurrentOffset + MipHeaderArray[MipIndex].Num() + sizeof(int64);
+			int64 BulkDataOffset = CurrentOffset + MipHeader.Num() + sizeof(int64);
 			MipHeaderWriter << BulkDataOffset;
 
-			OrderedBuffers.Add(FSharedBuffer::MakeView(MipHeaderArray[MipIndex].GetData(), MipHeaderArray[MipIndex].Num()));
-			CurrentOffset += MipHeaderArray[MipIndex].Num();
+			CurrentOffset += MipHeader.Num();
+			OrderedBuffers.Add(MakeSharedBufferFromArray(MoveTemp(MipHeader)));
 
 			if (bIsInlineMip)
 			{
@@ -424,8 +402,8 @@ void FTextureBuildFunction::Build(UE::DerivedData::FBuildContext& Context) const
 				CurrentOffset += CompressedMip.RawData.Num();
 			}
 
-
-			FMemoryWriter MipFooterWriter(MipFooterArray[MipIndex], /*bIsPersistent=*/ true);
+			TArray<uint8> MipFooter;
+			FMemoryWriter MipFooterWriter(MipFooter, /*bIsPersistent=*/ true);
 
 			MipFooterWriter << CompressedMip.SizeX;
 			MipFooterWriter << CompressedMip.SizeY;
@@ -441,19 +419,33 @@ void FTextureBuildFunction::Build(UE::DerivedData::FBuildContext& Context) const
 				DerivedDataKey << MipKeyPrefix;
 				DerivedDataKey.Appendf(TEXT("_MIP%u_%dx%d"), MipIndex, CompressedMip.SizeX, CompressedMip.SizeY);
 			}
-			FString DerivedDataKeyFinal (DerivedDataKey.ToString());
+			FString DerivedDataKeyFinal(DerivedDataKey.ToString());
 			MipFooterWriter << DerivedDataKeyFinal;
 
-			OrderedBuffers.Add(FSharedBuffer::MakeView(MipFooterArray[MipIndex].GetData(), MipFooterArray[MipIndex].Num()));
-			CurrentOffset += MipFooterArray[MipIndex].Num();
+			CurrentOffset += MipFooter.Num();
+			OrderedBuffers.Add(MakeSharedBufferFromArray(MoveTemp(MipFooter)));
 		}
 
-		// Bool serialized as 32-bit int is the footer  
+		// Bool serialized as 32-bit int is the footer
 		uint32 IsVirtual = 0;
 		OrderedBuffers.Add(FSharedBuffer::MakeView(&IsVirtual, sizeof(IsVirtual)));
 
-
 		FCompositeBuffer CompositeResult(OrderedBuffers);
-		Context.AddPayload(UE::DerivedData::FPayloadId::FromName(TEXT("Texture"_SV)), FCompressedBuffer::Compress(NAME_Default, CompositeResult));
+		Context.AddPayload(UE::DerivedData::FPayloadId::FromName(TEXT("Texture"_SV)), CompositeResult);
+	}
+
+	// Write out streaming mips as payloads.
+	const int32 WritableMipCount = MipCount - ((NumMipsInTail > 0) ? ((int32)NumMipsInTail - 1) : 0);
+	for (int32 MipIndex = 0; MipIndex < WritableMipCount; ++MipIndex)
+	{
+		if (MipIndex < FirstInlineMip)
+		{
+			FCompressedImage2D& CompressedMip = CompressedMips[MipIndex];
+			TStringBuilder<32> PayloadName;
+			PayloadName << "Mip" << MipIndex;
+
+			FSharedBuffer MipData = MakeSharedBufferFromArray(MoveTemp(CompressedMip.RawData));
+			Context.AddPayload(UE::DerivedData::FPayloadId::FromName(*PayloadName), MipData);
+		}
 	}
 }
diff --git a/Engine/Source/Runtime/Engine/Private/TextureDerivedDataTask.cpp b/Engine/Source/Runtime/Engine/Private/TextureDerivedDataTask.cpp
--- a/Engine/Source/Runtime/Engine/Private/TextureDerivedDataTask.cpp
+++ b/Engine/Source/Runtime/Engine/Private/TextureDerivedDataTask.cpp
@@ -48,6 +48,7 @@
 #include "Misc/ScopeRWLock.h"
 #include "ProfilingDebugging/CookStats.h"
 #include "Serialization/BulkDataRegistry.h"
+#include "Templates/CheckValueCast.h"
 #include "TextureDerivedDataBuildUtils.h"
 #include "VT/VirtualTextureChunkDDCCache.h"
 #include "VT/VirtualTextureDataBuilder.h"
@@ -287,7 +288,7 @@ public:
 		{
 			Source.OperateOnLoadedBulkData([&Buffer](const FSharedBuffer& BulkDataBuffer)
 			{
-				Buffer = FCompressedBuffer::Compress(NAME_Default, BulkDataBuffer);
+				Buffer = FCompressedBuffer::Compress(BulkDataBuffer);
 			});
 		}
 		return Buffer;
@@ -391,29 +392,32 @@ void FTextureCacheDerivedDataWorker::ConsumeBuildFunctionOutput(const UE::Derive
 				break;
 			}
 
-			TStringBuilder<32> PayloadName;
-			PayloadName << "Mip" << MipIndex;
-
-			FPayloadId MipPayloadId = FPayloadId::FromName(*PayloadName);
+			FPayloadId MipPayloadId = FPayloadId::FromName(WriteToString<8>(TEXT("Mip"), MipIndex));
 			const FPayload& MipPayload = BuildOutput.GetPayload(MipPayloadId);
 			if (MipPayload.IsNull())
 			{
 				UE_LOG(LogTexture, Warning, TEXT("Texture build function missing Mip%d output payload when building %s derived data for %s"), MipIndex, *BuildSettingsPerLayer[0].TextureFormatName.GetPlainNameString(), *Texture.GetPathName());
 				return;
 			}
+
 			FSharedBuffer MipData = MipPayload.GetData().Decompress();
-			TArrayView<const uint8> MipDataView((const uint8*)MipData.GetData(), MipData.GetSize());
-			GetDerivedDataCacheRef().Put(*Mip.DerivedDataKey, MipDataView, TexturePath, bReplaceExistingDDC);
+
+			{
+				TArray<uint8> MipDataCopy;
+				FMemoryWriter MipWriter(MipDataCopy, /*bIsPersistent=*/ true);
+				int32 MipSize = TCheckValueCast<int32>(MipData.GetSize());
+				MipWriter << MipSize;
+				MipWriter.Serialize(const_cast<void*>(MipData.GetData()), MipSize);
+				GetDerivedDataCacheRef().Put(*Mip.DerivedDataKey, MipDataCopy, TexturePath, bReplaceExistingDDC);
+			}
 
 			if (bInlineMips && (MipIndex >= (int32)BuildSettingsPerLayer[0].LODBiasWithCinematicMips))
 			{
-				int32 MipSize = 0;
-				FMemoryReaderView Ar(MipDataView, /*bIsPersistent=*/ true);
-				Ar << MipSize;
+				const uint64 MipSize = MipData.GetSize();
 
 				Mip.BulkData.Lock(LOCK_READ_WRITE);
-				void* MipAllocData = Mip.BulkData.Realloc(MipSize);
-				Ar.Serialize(MipAllocData, MipSize);
+				void* MipAllocData = Mip.BulkData.Realloc(int64(MipSize));
+				MakeMemoryView(MipAllocData, MipSize).CopyFrom(MipData);
 				Mip.BulkData.Unlock();
 				Mip.DerivedDataKey.Empty();
 			}
@@ -1218,14 +1222,11 @@ private:
 				if (bInlineMips && MipIndex >= FirstMipToLoad)
 				{
 					FSharedBuffer MipData = MipPayload.GetData().Decompress();
-					FMemoryReaderView Ar(MipData, /*bIsPersistent=*/ true);
-
-					int32 MipSize = 0;
-					Ar << MipSize;
+					const uint64 MipSize = MipData.GetSize();
 
 					Mip.BulkData.Lock(LOCK_READ_WRITE);
-					void* MipAllocData = Mip.BulkData.Realloc(MipSize);
-					Ar.Serialize(MipAllocData, MipSize);
+					void* MipAllocData = Mip.BulkData.Realloc(int64(MipSize));
+					MakeMemoryView(MipAllocData, MipSize).CopyFrom(MipData);
 					Mip.BulkData.Unlock();
 					Mip.DerivedDataKey.Empty();
 				}
