Bump version guid for all DDC2 texture build functions.

diff --git a/Engine/Source/Developer/TextureBuild/Private/TextureBuildFunction.cpp b/Engine/Source/Developer/TextureBuild/Private/TextureBuildFunction.cpp
--- a/Engine/Source/Developer/TextureBuild/Private/TextureBuildFunction.cpp
+++ b/Engine/Source/Developer/TextureBuild/Private/TextureBuildFunction.cpp
@@ -26,7 +26,7 @@ DEFINE_LOG_CATEGORY_STATIC(LogTextureBuildFunction, Log, All);
 // MUST have a corresponding change to this version. Individual texture formats have a version to
 // change that is specific to the format. A merge conflict affecting the version MUST be resolved
 // by generating a new version.
-static const FGuid TextureDerivedDataVersion(TEXT("d565a540-b94c-4c4a-9852-16e343aba4b5"));
+static const FGuid TextureDerivedDataVersion(TEXT("0b123ca4-66c4-4ea4-b269-550de0c32eb1"));
 
 #ifndef CASE_ENUM_TO_TEXT
 #define CASE_ENUM_TO_TEXT(txt) case txt: return TEXT(#txt);
