Ensure functional parity between legacy and new texture build paths for texture cube array assets (see CL17268789)

diff --git a/Engine/Source/Developer/TextureBuild/Private/TextureBuildFunction.cpp b/Engine/Source/Developer/TextureBuild/Private/TextureBuildFunction.cpp
--- a/Engine/Source/Developer/TextureBuild/Private/TextureBuildFunction.cpp
+++ b/Engine/Source/Developer/TextureBuild/Private/TextureBuildFunction.cpp
@@ -363,7 +363,7 @@ void FTextureBuildFunction::Build(UE::DerivedData::FBuildContext& Context) const
 		TextureHeaderWriter << CompressedMips[0].SizeX;
 		TextureHeaderWriter << CompressedMips[0].SizeY;
 
-		const int32 NumSlices = BuildSettings.bCubemap ? 6 : (BuildSettings.bVolume || BuildSettings.bTextureArray) ? CompressedMips[0].SizeZ : 1;
+		const int32 NumSlices = (BuildSettings.bVolume || BuildSettings.bTextureArray) ? CompressedMips[0].SizeZ : BuildSettings.bCubemap ? 6 : 1;
 		static constexpr uint32 BitMask_CubeMap = 1u << 31u;
 		static constexpr uint32 BitMask_HasOptData = 1u << 30u;
 		static constexpr uint32 BitMask_NumSlices = BitMask_HasOptData - 1u;
