Fix bug in texture build flow (new and legacy) caused by my previous changes to the build flow.  Can't completely remove all information on the texture mips (serialized and in memory) about the presence of a mip in derived data.  Instead of a string key, now keeping a boolean indicating the mip's status in DDC (agnostic of whether we're using the new or old flow).

diff --git a/Engine/Source/Developer/TextureBuild/Private/TextureBuildFunction.cpp b/Engine/Source/Developer/TextureBuild/Private/TextureBuildFunction.cpp
--- a/Engine/Source/Developer/TextureBuild/Private/TextureBuildFunction.cpp
+++ b/Engine/Source/Developer/TextureBuild/Private/TextureBuildFunction.cpp
@@ -26,7 +26,7 @@ DEFINE_LOG_CATEGORY_STATIC(LogTextureBuildFunction, Log, All);
 // MUST have a corresponding change to this version. Individual texture formats have a version to
 // change that is specific to the format. A merge conflict affecting the version MUST be resolved
 // by generating a new version.
-static const FGuid TextureDerivedDataVersion(TEXT("0b123ca4-66c4-4ea4-b269-550de0c32eb1"));
+static const FGuid TextureDerivedDataVersion(TEXT("392f5367-8112-4305-8cef-ae7897844779"));
 
 #ifndef CASE_ENUM_TO_TEXT
 #define CASE_ENUM_TO_TEXT(txt) case txt: return TEXT(#txt);
@@ -425,15 +425,8 @@ void FTextureBuildFunction::Build(UE::DerivedData::FBuildContext& Context) const
 			EFileRegionType FileRegion = FFileRegion::SelectType(EPixelFormat(CompressedMip.PixelFormat));
 			MipFooterWriter << FileRegion;
 
-			TStringBuilder<512> DerivedDataKey;
-			if (!bIsInlineMip)
-			{
-				// This has to match the behavior in GetTextureDerivedMipKey in TextureDerivedData.cpp
-				DerivedDataKey << MipKeyPrefix;
-				DerivedDataKey.Appendf(TEXT("_MIP%u_%dx%d"), MipIndex, CompressedMip.SizeX, CompressedMip.SizeY);
-			}
-			FString DerivedDataKeyFinal(DerivedDataKey.ToString());
-			MipFooterWriter << DerivedDataKeyFinal;
+			bool bExistsInDerivedData = !bIsInlineMip;
+			MipFooterWriter << bExistsInDerivedData;
 
 			CurrentOffset += MipFooter.Num();
 			OrderedBuffers.Add(MakeSharedBufferFromArray(MoveTemp(MipFooter)));
diff --git a/Engine/Source/Runtime/Engine/Private/Texture2D.cpp b/Engine/Source/Runtime/Engine/Private/Texture2D.cpp
--- a/Engine/Source/Runtime/Engine/Private/Texture2D.cpp
+++ b/Engine/Source/Runtime/Engine/Private/Texture2D.cpp
@@ -165,17 +165,7 @@ void FTexture2DMipMap::SerializeCommon(FArchive& Ar, UObject* Owner, int32 MipId
 	if (!bCooked)
 	{
 		Ar << FileRegionType;
-
-		FString LegacyDDCKeyString;
-		if (Ar.IsSaving())
-		{
-			if (IsPagedToDerivedData())
-			{
-				LegacyDDCKeyString = TEXT("DummyLegacyKey");
-			}
-		}
-
-		Ar << LegacyDDCKeyString;
+		Ar << bPagedToDerivedData;
 	}
 #endif // #if WITH_EDITORONLY_DATA
 }
@@ -206,6 +196,7 @@ uint32 FTexture2DMipMap::StoreInDerivedDataCache(const FString& InDerivedDataKey
 	}
 	const uint32 Result = DerivedData.Num();
 	GetDerivedDataCacheRef().Put(*InDerivedDataKey, DerivedData, TextureName, bReplaceExistingDDC);
+	bPagedToDerivedData = true;
 	BulkData.RemoveBulkData();
 	return Result;
 }
diff --git a/Engine/Source/Runtime/Engine/Private/TextureDerivedData.cpp b/Engine/Source/Runtime/Engine/Private/TextureDerivedData.cpp
--- a/Engine/Source/Runtime/Engine/Private/TextureDerivedData.cpp
+++ b/Engine/Source/Runtime/Engine/Private/TextureDerivedData.cpp
@@ -55,7 +55,7 @@
 // In case of merge conflicts with DDC versions, you MUST generate a new GUID and set this new
 // guid as version
 
-#define TEXTURE_DERIVEDDATA_VER		TEXT("E405EE6911174BB2A4B24BAC98C41448")
+#define TEXTURE_DERIVEDDATA_VER		TEXT("6F4DB6EC89FB4C34B2415FC04F995708")
 
 // This GUID is mixed into DDC version for virtual textures only, this allows updating DDC version for VT without invalidating DDC for all textures
 // This is useful during development, but once large numbers of VT are present in shipped content, it will have the same problem as TEXTURE_DERIVEDDATA_VER
@@ -1107,6 +1107,7 @@ bool FTexturePlatformData::TryInlineMipData(int32 FirstMipToLoad, UTexture* Text
 				void* MipData = Mip.BulkData.Realloc(MipSize);
 				Ar.Serialize(MipData, MipSize);
 				Mip.BulkData.Unlock();
+				Mip.SetPagedToDerivedData(false);
 			}
 			else
 			{
diff --git a/Engine/Source/Runtime/Engine/Private/TextureDerivedDataTask.cpp b/Engine/Source/Runtime/Engine/Private/TextureDerivedDataTask.cpp
--- a/Engine/Source/Runtime/Engine/Private/TextureDerivedDataTask.cpp
+++ b/Engine/Source/Runtime/Engine/Private/TextureDerivedDataTask.cpp
@@ -447,6 +447,7 @@ void FTextureCacheDerivedDataWorker::ConsumeBuildFunctionOutput(const UE::Derive
 				void* MipAllocData = Mip.BulkData.Realloc(int64(MipSize));
 				MakeMemoryView(MipAllocData, MipSize).CopyFrom(MipData);
 				Mip.BulkData.Unlock();
+				Mip.SetPagedToDerivedData(false);
 			}
 		}
 
@@ -1393,6 +1394,7 @@ private:
 					void* MipAllocData = Mip.BulkData.Realloc(int64(MipSize));
 					MakeMemoryView(MipAllocData, MipSize).CopyFrom(MipData);
 					Mip.BulkData.Unlock();
+					Mip.SetPagedToDerivedData(false);
 				}
 			}
 			else
diff --git a/Engine/Source/Runtime/Engine/Public/TextureResource.h b/Engine/Source/Runtime/Engine/Public/TextureResource.h
--- a/Engine/Source/Runtime/Engine/Public/TextureResource.h
+++ b/Engine/Source/Runtime/Engine/Public/TextureResource.h
@@ -67,7 +67,11 @@ struct FTexture2DMipMap
 	/** The file region type appropriate for this mip's pixel format. */
 	EFileRegionType FileRegionType = EFileRegionType::None;
 
-	bool IsPagedToDerivedData() const { return BulkData.GetBulkDataSize() == 0; }
+	/** Whether this mip is stored in the derived data cache. */
+	bool bPagedToDerivedData = false;
+
+	bool IsPagedToDerivedData() const { return bPagedToDerivedData; }
+	void SetPagedToDerivedData(bool InValue) { bPagedToDerivedData = InValue; }
 
 	/**
 	 * Place mip-map data in the derived data cache associated with the provided
@@ -77,7 +81,6 @@ struct FTexture2DMipMap
 
 #endif // #if WITH_EDITORONLY_DATA
 private:
-
 	void SerializeCommon(FArchive& Ar, UObject* Owner, int32 MipIndex, bool bBulkDataConditionalOnKey);
 };
 
