Texture: Fixed the texture build function to stop allocating double the required memory for the source texture

#rnx

diff --git a/Engine/Source/Developer/TextureBuild/Private/TextureBuildFunction.cpp b/Engine/Source/Developer/TextureBuild/Private/TextureBuildFunction.cpp
--- a/Engine/Source/Developer/TextureBuild/Private/TextureBuildFunction.cpp
+++ b/Engine/Source/Developer/TextureBuild/Private/TextureBuildFunction.cpp
@@ -267,6 +267,7 @@ static bool TryReadTextureSourceFromCompactBinary(FCbFieldView Source, UE::Deriv
 		}
 		else
 		{
+			SourceMip->RawData.Reset(MipSize);
 			SourceMip->RawData.AddUninitialized(MipSize);
 			FMemory::Memcpy(
 				SourceMip->RawData.GetData(),
