fix Volume Textures with LeaveExistingMips
was not correctly mipping down Volume Depth in TextureSources
need to correct the calculation in various places where its code duped
add some size checking

diff --git a/Engine/Source/Developer/TextureBuild/Private/TextureBuildFunction.cpp b/Engine/Source/Developer/TextureBuild/Private/TextureBuildFunction.cpp
--- a/Engine/Source/Developer/TextureBuild/Private/TextureBuildFunction.cpp
+++ b/Engine/Source/Developer/TextureBuild/Private/TextureBuildFunction.cpp
@@ -110,6 +110,7 @@ static FTextureBuildSettings ReadBuildSettingsFromCompactBinary(const FCbObjectV
 	ReadCbField(Object["DiffuseConvolveMipLevel"], BuildSettings.DiffuseConvolveMipLevel);
 	ReadCbField(Object["SharpenMipKernelSize"], BuildSettings.SharpenMipKernelSize);
 	ReadCbField(Object["MaxTextureResolution"], BuildSettings.MaxTextureResolution);
+	check( BuildSettings.MaxTextureResolution != 0 );
 	ReadCbField(Object["TextureFormatName"], BuildSettings.TextureFormatName);
 	ReadCbField(Object["bHDRSource"], BuildSettings.bHDRSource);
 	ReadCbField(Object["MipGenSettings"], BuildSettings.MipGenSettings);
@@ -180,7 +181,8 @@ static ERawImageFormat::Type ComputeRawImageFormat(ETextureSourceFormat SourceFo
 	return FImageCoreUtils::ConvertToRawImageFormat(SourceFormat);
 }
 
-static bool TryReadTextureSourceFromCompactBinary(FCbFieldView Source, UE::DerivedData::FBuildContext& Context, TArray<FImage>& OutMips)
+static bool TryReadTextureSourceFromCompactBinary(FCbFieldView Source, UE::DerivedData::FBuildContext& Context,
+												const FTextureBuildSettings & BuildSettings, TArray<FImage>& OutMips)
 {
 	FSharedBuffer InputBuffer = Context.FindInput(Source.GetName());
 	if (!InputBuffer)
@@ -207,6 +209,8 @@ static bool TryReadTextureSourceFromCompactBinary(FCbFieldView Source, UE::Deriv
 	int32 MipSizeY = SizeY;
 
 	const uint8* DecompressedSourceData = (const uint8*)InputBuffer.GetData();
+	int64 DecompressedSourceDataSize = InputBuffer.GetSize();
+
 	TArray64<uint8> IntermediateDecompressedData;
 	if (CompressionFormat != TSCF_None)
 	{
@@ -232,6 +236,7 @@ static bool TryReadTextureSourceFromCompactBinary(FCbFieldView Source, UE::Deriv
 			return false;
 		}
 		DecompressedSourceData = IntermediateDecompressedData.GetData();
+		DecompressedSourceDataSize = IntermediateDecompressedData.Num();
 		InputBuffer.Reset();
 	}
 
@@ -250,9 +255,13 @@ static bool TryReadTextureSourceFromCompactBinary(FCbFieldView Source, UE::Deriv
 			GammaSpace
 		);
 
+		check( MipOffset + MipSize <= DecompressedSourceDataSize );
+		check( SourceMip->GetImageSizeBytes() == MipSize );
+
 		if ((MipsCbArrayView.Num() == 1) && (CompressionFormat != TSCF_None))
 		{
 			// In the case where there is only one mip and its already in a TArray, there is no need to allocate new array contents, just use a move instead
+			check( MipOffset == 0 );
 			SourceMip->RawData = MoveTemp(IntermediateDecompressedData);
 		}
 		else
@@ -268,6 +277,10 @@ static bool TryReadTextureSourceFromCompactBinary(FCbFieldView Source, UE::Deriv
 
 		MipSizeX = FMath::Max(MipSizeX / 2, 1);
 		MipSizeY = FMath::Max(MipSizeY / 2, 1);
+		if ( BuildSettings.bVolume )
+		{
+			NumSlices = FMath::Max(NumSlices / 2, 1);
+		}
 	}
 
 	return true;
@@ -342,14 +355,14 @@ void FTextureBuildFunction::Build(UE::DerivedData::FBuildContext& Context) const
 	ReadOutputSettingsFromCompactBinary(Settings["Output"].AsObjectView(), NumInlineMips);
 
 	TArray<FImage> SourceMips;
-	if (!TryReadTextureSourceFromCompactBinary(Settings["Source"], Context, SourceMips))
+	if (!TryReadTextureSourceFromCompactBinary(Settings["Source"], Context,BuildSettings, SourceMips))
 	{
 		return;
 	}
 
 	TArray<FImage> AssociatedNormalSourceMips;
 	if (FCbFieldView CompositeSource = Settings["CompositeSource"];
-		CompositeSource && !TryReadTextureSourceFromCompactBinary(CompositeSource, Context, AssociatedNormalSourceMips))
+		CompositeSource && !TryReadTextureSourceFromCompactBinary(CompositeSource, Context,BuildSettings, AssociatedNormalSourceMips))
 	{
 		return;
 	}
@@ -359,6 +372,14 @@ void FTextureBuildFunction::Build(UE::DerivedData::FBuildContext& Context) const
 	TArray<FCompressedImage2D> CompressedMips;
 	uint32 NumMipsInTail;
 	uint32 ExtData;
+	
+	/*
+	for( const FImage & Image : SourceMips )
+	{
+		UE_LOG(LogInit,Display,TEXT("FTextureBuildFunction: SourceMips : %dx%dx%d"),
+			Image.SizeX,Image.SizeY,Image.NumSlices);
+	}
+	*/
 
 	bool bBuildSucceeded = FModuleManager::GetModuleChecked<ITextureCompressorModule>(TEXTURE_COMPRESSOR_MODULENAME).BuildTexture(
 		SourceMips,
@@ -385,6 +406,7 @@ void FTextureBuildFunction::Build(UE::DerivedData::FBuildContext& Context) const
 		DescriptionWriter.BeginArray("Size"_ASV);
 		DescriptionWriter.AddInteger(CompressedMips[0].SizeX);
 		DescriptionWriter.AddInteger(CompressedMips[0].SizeY);
+		// this is brittle ; CompressedMips[] should report NumSlices itself
 		const int32 NumSlices = (BuildSettings.bVolume || BuildSettings.bTextureArray) ? CompressedMips[0].SizeZ : BuildSettings.bCubemap ? 6 : 1;
 		DescriptionWriter.AddInteger(NumSlices);
 		DescriptionWriter.EndArray();
@@ -410,6 +432,11 @@ void FTextureBuildFunction::Build(UE::DerivedData::FBuildContext& Context) const
 			DescriptionWriter.AddInteger(CompressedMip.SizeZ);
 			DescriptionWriter.EndArray();
 
+			/*
+			UE_LOG(LogInit,Display,TEXT("FTextureBuildFunction: CompressedMip: %d = %dx%dx%d"),MipIndex,
+				CompressedMip.SizeX,CompressedMip.SizeY,CompressedMip.SizeZ);
+			*/
+
 			const bool bIsInlineMip = MipIndex >= FirstInlineMip;
 
 			EFileRegionType FileRegion = FFileRegion::SelectType(EPixelFormat(CompressedMip.PixelFormat));
diff --git a/Engine/Source/Developer/TextureCompressor/Private/TextureCompressorModule.cpp b/Engine/Source/Developer/TextureCompressor/Private/TextureCompressorModule.cpp
--- a/Engine/Source/Developer/TextureCompressor/Private/TextureCompressorModule.cpp
+++ b/Engine/Source/Developer/TextureCompressor/Private/TextureCompressorModule.cpp
@@ -3025,7 +3025,7 @@ private:
 				const FImage& BaseImage = pSourceMips->Last();
 				bool bSuitableFormat = BaseImage.Format == ERawImageFormat::RGBA32F;
 			
-				check( MaxTextureResolution > 1 );
+				check( MaxTextureResolution > 0 );
 				check( BaseImage.SizeX > MaxTextureResolution || 
 					   BaseImage.SizeY > MaxTextureResolution );
 
diff --git a/Engine/Source/Editor/UnrealEd/Private/Factories/EditorFactories.cpp b/Engine/Source/Editor/UnrealEd/Private/Factories/EditorFactories.cpp
--- a/Engine/Source/Editor/UnrealEd/Private/Factories/EditorFactories.cpp
+++ b/Engine/Source/Editor/UnrealEd/Private/Factories/EditorFactories.cpp
@@ -3578,18 +3578,6 @@ UTexture * UTextureFactory::ImportDDS(const uint8* Buffer,int64 Length,UObject*
 		MipCount = MAX_TEXTURE_MIP_COUNT;
 	}
 
-	#if 1
-	if ( DDS->Dimension == 3 )
-	{
-		// Unreal doesn't handle volume mips correctly?
-		//  (slices doesn't mip down)
-		// force mip import off for now
-		// @@!! fix Unreal volume mips?
-		// see also GetMippedNumSlices
-		MipCount = 1;
-	}
-	#endif
-
 	Texture->Source.Init(DDS->Width,DDS->Height,NumSlices,MipCount,TSFormat);
 
 	for(int MipIndex=0;MipIndex<MipCount;MipIndex++)
@@ -3599,20 +3587,6 @@ UTexture * UTextureFactory::ImportDDS(const uint8* Buffer,int64 Length,UObject*
 		if ( DDS->Dimension == 3 )
 		{
 			check( DDS->Mips.Num() == DDS->MipCount );
-
-			#if 0
-			// hack for Unreal not mipping volumes correctly (NumSlices doesn't go down like it should)
-			//  (moot because we now force MipCount to 1 for volumes)
-			// @@!! if the fix is made in GetMippedNumSlices , this is unnecessary
-			if ( MipIndex > 0 && MipLock.Image.NumSlices > (int32)DDS->Mips[MipIndex].Depth )
-			{
-				// memset the miplock to zero : last slices won't be written :
-				memset(MipLock.Image.RawData,0,MipLock.Image.GetImageSizeBytes());
-				// only copy the portion the DDS has :
-				MipLock.Image.NumSlices = DDS->Mips[MipIndex].Depth;
-			}
-			#endif
-
 			check( DDS->Mips[MipIndex].Depth == MipLock.Image.NumSlices );
 
 			if ( ! DDS->GetMipImage( MipLock.Image, MipIndex ) )
diff --git a/Engine/Source/Runtime/Core/Public/Math/Color.h b/Engine/Source/Runtime/Core/Public/Math/Color.h
--- a/Engine/Source/Runtime/Core/Public/Math/Color.h
+++ b/Engine/Source/Runtime/Core/Public/Math/Color.h
@@ -23,6 +23,8 @@ enum class EGammaSpace : uint8
 	Pow22,
 	/** Use the standard sRGB conversion. */
 	sRGB,
+
+	Invalid
 };
 
 
diff --git a/Engine/Source/Runtime/Engine/Classes/Engine/Texture.h b/Engine/Source/Runtime/Engine/Classes/Engine/Texture.h
--- a/Engine/Source/Runtime/Engine/Classes/Engine/Texture.h
+++ b/Engine/Source/Runtime/Engine/Classes/Engine/Texture.h
@@ -339,9 +339,12 @@ struct FTextureSource
 	FORCEINLINE ETextureSourceCompressionFormat GetSourceCompression() const { return CompressionFormat; }
 	FORCEINLINE bool IsSourceCompressed() const { return GetSourceCompression() != ETextureSourceCompressionFormat::TSCF_None; }
 
-	/** Get GammaSpace for this Format */
+	/** Get GammaSpace for this Source (asks owner) */
 	ENGINE_API EGammaSpace GetGammaSpace(int LayerIndex) const;
 	
+	/** Get TextureClass for this Source (asks owner) */
+	ENGINE_API ETextureClass GetTextureClass() const;
+
 	/** Get mipped number of slices (volumes mip but arrays don't) */
 	ENGINE_API int GetMippedNumSlices(int NumSlices,int MipIndex) const;
 
@@ -551,6 +554,11 @@ private:
 	TDontCopy<FRWLock> BulkDataLock;
 	/** Owner for associating BulkData with a package */
 	UTexture* Owner;
+	/** TextureClass == Owner->GetTextureClass(); a copy is kept here for torn-off **/
+	ETextureClass TornOffTextureClass;
+	/** if Owner != null, check Owner->GetGammaSpace , if it is null, use TornOffGammaSpace
+	* do not check this directly, use GetGammaSpace **/
+	EGammaSpace TornOffGammaSpace;
 #endif
 	/** The bulk source data. */
 	UE::Serialization::FEditorBulkData BulkData;
@@ -853,6 +861,7 @@ public:
 	 * @returns true if all requested mips have been loaded.
 	 */
 	bool TryLoadMips(int32 FirstMipToLoad, void** OutMipData, FStringView DebugContext);
+	bool TryLoadMipsWithSizes(int32 FirstMipToLoad, void** OutMipData, int64 * OutMipSize, FStringView DebugContext);
 
 	/** Serialization. */
 	void Serialize(FArchive& Ar, class UTexture* Owner);
diff --git a/Engine/Source/Runtime/Engine/Private/Rendering/Texture3DResource.cpp b/Engine/Source/Runtime/Engine/Private/Rendering/Texture3DResource.cpp
--- a/Engine/Source/Runtime/Engine/Private/Rendering/Texture3DResource.cpp
+++ b/Engine/Source/Runtime/Engine/Private/Rendering/Texture3DResource.cpp
@@ -33,9 +33,10 @@ void FVolumeTextureBulkData::MergeMips(int32 NumMips)
 {
 	check(NumMips < MAX_TEXTURE_MIP_COUNT);
 
-	uint64 MergedSize = 0;
+	int64 MergedSize = 0;
 	for (int32 MipIndex = FirstMipIdx; MipIndex < NumMips; ++MipIndex)
 	{
+		check( MipSize[MipIndex] != 0 );
 		MergedSize += MipSize[MipIndex];
 	}
 
@@ -69,9 +70,13 @@ FTexture3DResource::FTexture3DResource(UVolumeTexture* InOwner, const FStreamabl
 , InitialData(InState.RequestedFirstLODIdx())
 {
 	const int32 FirstLODIdx = InState.RequestedFirstLODIdx();
-	if (PlatformData && const_cast<FTexturePlatformData*>(PlatformData)->TryLoadMips(FirstLODIdx + InState.AssetLODBias, InitialData.GetMipData() + FirstLODIdx, InOwner->GetPathName()))
+
+	// changed to use TryLoadMipsWithSizes
+	if (PlatformData && const_cast<FTexturePlatformData*>(PlatformData)->TryLoadMipsWithSizes(FirstLODIdx + InState.AssetLODBias, 
+		InitialData.GetMipData() + FirstLODIdx, InitialData.GetMipSize() + FirstLODIdx, InOwner->GetPathName()))
 	{
 		// Compute the size of each mips so that they can be merged into a single allocation.
+		// -> this should be unnecessary now that TryLoadMipsWithSizes reports the sizes
 		if (GUseTexture3DBulkDataRHI)
 		{
 			for (int32 MipIndex = FirstLODIdx; MipIndex < InState.MaxNumLODs; ++MipIndex)
@@ -86,12 +91,17 @@ FTexture3DResource::FTexture3DResource(UVolumeTexture* InOwner, const FStreamabl
 
 				uint32 TextureAlign = 0;
 				uint64 PlatformMipSize = RHICalcTexture3DPlatformSize(MipExtentX, MipExtentY, MipExtentZ, (EPixelFormat)PixelFormat, State.NumRequestedLODs, CreationFlags, FRHIResourceCreateInfo(PlatformData->GetExtData()), TextureAlign);
+				
+				UE_LOG(LogTexture,Verbose,TEXT("FTexture3DResource::FTexture3DResource %d : %dx%dx%d : MipBytes=%d"),
+					MipIndex,MipExtentX,MipExtentY,MipExtentZ,
+					(int)PlatformMipSize);
 
 				// The bulk data can be bigger because of memory alignment constraints on each slice and mips.
-				InitialData.GetMipSize()[MipIndex] = FMath::Max<uint64>(
-					MipMap.BulkData.GetBulkDataSize(), 
-					CalcTextureMipMapSize3D(SizeX, SizeY, SizeZ, (EPixelFormat)PixelFormat, MipIndex)
-				);
+				int64 BulkDataSize = MipMap.BulkData.GetBulkDataSize();
+				int64 CalcMipSize = CalcTextureMipMapSize3D(SizeX, SizeY, SizeZ, (EPixelFormat)PixelFormat, MipIndex);
+				check( BulkDataSize >= CalcMipSize );
+				//InitialData.GetMipSize()[MipIndex] = BulkDataSize;
+				check( InitialData.GetMipSize()[MipIndex] == BulkDataSize );
 			}
 		}
 
@@ -140,6 +150,18 @@ void FTexture3DResource::CreateTexture()
 				// We check if this is really the rendering thread to find out if the engine is initializing.
 				const uint32 NumBlockX = (uint32)FMath::DivideAndRoundUp<int32>(Mip.SizeX, BlockSizeX);
 				const uint32 NumBlockY = (uint32)FMath::DivideAndRoundUp<int32>(Mip.SizeY, BlockSizeY);
+
+				{
+				int64 MipBytes = InitialData.GetMipSize()[ResourceMipIdx];
+				int64 UploadSize = NumBlockX * NumBlockY * BlockBytes * Mip.SizeZ;
+				check( MipBytes >= UploadSize );
+
+				UE_LOG(LogTexture,Verbose,TEXT("FTexture3DResource::CreateTexture:RHIUpdateTexture3D %d : %dx%dx%d : RowStride=%d, SliceStride=%d, MipBytes=%d"),
+					ResourceMipIdx,Mip.SizeX,Mip.SizeY,Mip.SizeZ,
+					NumBlockX * BlockBytes, NumBlockX * NumBlockY * BlockBytes,
+					(int)MipBytes);
+				}
+
 				RHIUpdateTexture3D(Texture3DRHI, RHIMipIdx, UpdateRegion, NumBlockX * BlockBytes, NumBlockX * NumBlockY * BlockBytes, MipData);
 			}
 		}
diff --git a/Engine/Source/Runtime/Engine/Private/Rendering/Texture3DResource.h b/Engine/Source/Runtime/Engine/Private/Rendering/Texture3DResource.h
--- a/Engine/Source/Runtime/Engine/Private/Rendering/Texture3DResource.h
+++ b/Engine/Source/Runtime/Engine/Private/Rendering/Texture3DResource.h
@@ -50,13 +50,13 @@ public:
 	void MergeMips(int32 NumMips);
 
 	void** GetMipData() { return MipData; }
-	uint64* GetMipSize() { return MipSize; }
+	int64* GetMipSize() { return MipSize; }
 	int32 GetFirstMipIdx() const { return FirstMipIdx; }
 
 protected:
 
 	void* MipData[MAX_TEXTURE_MIP_COUNT];
-	uint64 MipSize[MAX_TEXTURE_MIP_COUNT];
+	int64 MipSize[MAX_TEXTURE_MIP_COUNT];
 	int32 FirstMipIdx;
 };
 
diff --git a/Engine/Source/Runtime/Engine/Private/Texture.cpp b/Engine/Source/Runtime/Engine/Private/Texture.cpp
--- a/Engine/Source/Runtime/Engine/Private/Texture.cpp
+++ b/Engine/Source/Runtime/Engine/Private/Texture.cpp
@@ -1300,6 +1300,8 @@ FTextureSource::FTextureSource()
 	: 
 #if WITH_EDITOR
 	  Owner(nullptr),
+	  TornOffTextureClass(ETextureClass::Invalid),
+	  TornOffGammaSpace(EGammaSpace::Invalid),
 #endif
 	  NumLockedMips(0u)
 	, LockState(ELockState::None)
@@ -1343,15 +1345,13 @@ int32 FTextureSource::GetBytesPerPixel(ETextureSourceFormat Format)
 
 bool FTextureSource::IsCubeOrCubeArray() const
 {
-	check( Owner != nullptr );
-	ETextureClass TextureClass = Owner->GetTextureClass();
+	ETextureClass TextureClass = GetTextureClass();
 	return TextureClass == ETextureClass::Cube || TextureClass == ETextureClass::CubeArray;
 }
 
 bool FTextureSource::IsVolume() const
 {
-	check( Owner != nullptr );
-	ETextureClass TextureClass = Owner->GetTextureClass();
+	ETextureClass TextureClass = GetTextureClass();
 	return TextureClass == ETextureClass::Volume;
 }
 
@@ -1611,6 +1611,10 @@ FTextureSource FTextureSource::CopyTornOff() const
 	// Use the default copy constructor to copy all the fields without having to write them manually
 	Result = *this;
 	Result.Owner = nullptr; // TornOffs don't count as belonging to the same owner
+	// Result can't talk to Owner any more, so save info we need :
+	check( Owner != nullptr );
+	Result.TornOffGammaSpace = Owner->GetGammaSpace();
+	Result.TornOffTextureClass = Owner->GetTextureClass();
 	return Result;
 }
 
@@ -2017,18 +2021,9 @@ FString FTextureSource::GetIdString() const
 
 int FTextureSource::GetMippedNumSlices(int InNumSlices,int InMipIndex) const
 {
-	// @@!! test me!  before was always returning InNumSlices
-	//  not sure you can change this, does it break old content?
-	//  but I think any old content that had TextureSource mips for volumes was broken anyway
-	//	so maybe it's okay?
-
-	#if 1
-
-	// match old behavior
-	return InNumSlices;
+	// old behavior was to not mip down NumSlices in TextureSource for volume SizeZ
+	//return InNumSlices;
 
-	#else
-	
 	// what to do with NumSlices on the mip ?
 	// if this is an Array, it should stay the same
 	// if this is a Volume, it should mip down
@@ -2036,8 +2031,14 @@ int FTextureSource::GetMippedNumSlices(int InNumSlices,int InMipIndex) const
 	// TextureSource does not know if it's a volume or not
 	// need to check type of Owner Texture
 
-	bool bIsVolume = Owner ? Owner->IsA(UVolumeTexture::StaticClass()) : false;
-	if ( bIsVolume )
+	check( InNumSlices > 0 );
+	// fast path short cut : 1 slice is always 1 slice
+	if ( InNumSlices == 1 )
+	{
+		return 1;
+	}
+
+	if ( IsVolume() )
 	{
 		return FMath::Max(InNumSlices >> InMipIndex, 1);
 	}
@@ -2045,12 +2046,26 @@ int FTextureSource::GetMippedNumSlices(int InNumSlices,int InMipIndex) const
 	{
 		return InNumSlices;
 	}	
+}
 
-	#endif
+ETextureClass FTextureSource::GetTextureClass() const
+{
+	// TextureSource does not know its own class, but its owning Texture does :
+	if ( Owner != nullptr )
+	{
+		return Owner->GetTextureClass();
+	}
+	else
+	{
+		// torn off, should have saved TornOffTextureClass
+		check( TornOffTextureClass != ETextureClass::Invalid );
+		return TornOffTextureClass;
+	}
 }
 
 EGammaSpace FTextureSource::GetGammaSpace(int LayerIndex) const
 {
+	// Texture->GetGammaSpace does not validate against format, but I do? a bit weird, fix that
 	if ( ! ERawImageFormat::GetFormatNeedsGammaSpace( FImageCoreUtils::ConvertToRawImageFormat(GetFormat(LayerIndex)) ) )
 	{
 		return EGammaSpace::Linear;
@@ -2063,9 +2078,9 @@ EGammaSpace FTextureSource::GetGammaSpace(int LayerIndex) const
 	}
 	else
 	{
-		// TextureSource does not have gamma information, just guess it ?
-		// we know format GetFormatNeedsGammaSpace is true, so default is SRGB
-		return EGammaSpace::sRGB;
+		// torn off, should have saved TornOffGammaSpace
+		check( TornOffGammaSpace != EGammaSpace::Invalid );
+		return TornOffGammaSpace;
 	}
 }
 
diff --git a/Engine/Source/Runtime/Engine/Private/TextureDerivedData.cpp b/Engine/Source/Runtime/Engine/Private/TextureDerivedData.cpp
--- a/Engine/Source/Runtime/Engine/Private/TextureDerivedData.cpp
+++ b/Engine/Source/Runtime/Engine/Private/TextureDerivedData.cpp
@@ -87,7 +87,8 @@ namespace TextureCookStats
 ------------------------------------------------------------------------------*/
 
 /**
- * Serialize build settings for use when generating the derived data key.
+ * Serialize build settings for use when generating the derived data key. (DDC1)
+ * Must keep in sync with DDC2 Key WriteBuildSettings
  */
 static void SerializeForKey(FArchive& Ar, const FTextureBuildSettings& Settings)
 {
@@ -217,6 +218,12 @@ static void SerializeForKey(FArchive& Ar, const FTextureBuildSettings& Settings)
 		Ar << TempGuid;
 	}
 
+	if ( Settings.bVolume )
+	{
+		TempGuid = FGuid(0xCC4348B8,0x84714993,0xAB1E2C93,0x8EA6C9E0);
+		Ar << TempGuid;
+	}
+
 	// Note - compression quality is added to the DDC by the formats (based on whether they
 	// use them or not).
 	// This is true for:
@@ -737,7 +744,8 @@ static void GetTextureBuildSettings(
 	OutBuildSettings.BlueChromaticityCoordinate = FVector2f(Texture.SourceColorSettings.BlueChromaticityCoordinate);
 	OutBuildSettings.WhiteChromaticityCoordinate = FVector2f(Texture.SourceColorSettings.WhiteChromaticityCoordinate);
 	OutBuildSettings.ChromaticAdaptationMethod = static_cast<uint8>(Texture.SourceColorSettings.ChromaticAdaptationMethod);
-
+	
+	check( OutBuildSettings.MaxTextureResolution == FTextureBuildSettings::MaxTextureResolutionDefault );
 	if (Texture.MaxTextureSize > 0)
 	{
 		OutBuildSettings.MaxTextureResolution = Texture.MaxTextureSize;
@@ -1866,6 +1874,14 @@ bool FTexturePlatformData::IsReadyForAsyncPostLoad() const
 }
 
 bool FTexturePlatformData::TryLoadMips(int32 FirstMipToLoad, void** OutMipData, FStringView DebugContext)
+{
+	// TryLoadMips fills mip pointers but not sizes
+	//  dangerous, not robust, use TryLoadMipsWithSizes instead
+
+	return TryLoadMipsWithSizes(FirstMipToLoad,OutMipData,nullptr,DebugContext);
+}
+
+bool FTexturePlatformData::TryLoadMipsWithSizes(int32 FirstMipToLoad, void** OutMipData, int64 * OutMipSize, FStringView DebugContext)
 {
 	TRACE_CPUPROFILER_EVENT_SCOPE(FTexturePlatformData::TryLoadMips);
 
@@ -1879,7 +1895,7 @@ bool FTexturePlatformData::TryLoadMips(int32 FirstMipToLoad, void** OutMipData,
 	FAsyncMipHandles AsyncHandles;
 	FDerivedDataCacheInterface& DDC = GetDerivedDataCacheRef();
 	if (!BeginLoadDerivedMips(*this, FirstMipToLoad, DebugContext, AsyncHandles,
-		[this, OutMipData, FirstMipToLoad, &NumMipsCached](int32 MipIndex, FSharedBuffer MipBuffer)
+		[this, OutMipData,OutMipSize, FirstMipToLoad, &NumMipsCached](int32 MipIndex, FSharedBuffer MipBuffer)
 		{
 			FTexture2DMipMap& Mip = Mips[MipIndex];
 
@@ -1892,6 +1908,10 @@ bool FTexturePlatformData::TryLoadMips(int32 FirstMipToLoad, void** OutMipData,
 				OutMipData[MipIndex - FirstMipToLoad] = FMemory::Malloc(MipSize);
 				FMemory::Memcpy(OutMipData[MipIndex - FirstMipToLoad], MipBuffer.GetData(), MipSize);
 			}
+			if ( OutMipSize ) 
+			{
+				OutMipSize[MipIndex - FirstMipToLoad] = MipSize;
+			}
 		}))
 	{
 		return false;
@@ -1932,6 +1952,10 @@ bool FTexturePlatformData::TryLoadMips(int32 FirstMipToLoad, void** OutMipData,
 #endif
 				Mip.BulkData.GetCopy(&OutMipData[MipIndex - FirstMipToLoad], true);
 			}
+			if ( OutMipSize )
+			{
+				OutMipSize[MipIndex - FirstMipToLoad] = BulkDataSize;
+			}
 			NumMipsCached++;
 		}
 	}
@@ -1958,6 +1982,10 @@ bool FTexturePlatformData::TryLoadMips(int32 FirstMipToLoad, void** OutMipData,
 						OutMipData[MipIndex - FirstMipToLoad] = FMemory::Malloc(TempData.Num());
 						Ar.Serialize(OutMipData[MipIndex - FirstMipToLoad], TempData.Num());
 					}
+					if ( OutMipSize )
+					{
+						OutMipSize[MipIndex - FirstMipToLoad] = TempData.Num();
+					}
 				}
 				else
 				{
diff --git a/Engine/Source/Runtime/Engine/Private/TextureDerivedDataBuildUtils.cpp b/Engine/Source/Runtime/Engine/Private/TextureDerivedDataBuildUtils.cpp
--- a/Engine/Source/Runtime/Engine/Private/TextureDerivedDataBuildUtils.cpp
+++ b/Engine/Source/Runtime/Engine/Private/TextureDerivedDataBuildUtils.cpp
@@ -130,6 +130,10 @@ static void WriteBuildSettings(FCbWriter& Writer, const FTextureBuildSettings& B
 	WriteCbFieldWithDefault<bool>(Writer, "bCubemap", BuildSettings.bCubemap, DefaultSettings.bCubemap);
 	WriteCbFieldWithDefault<bool>(Writer, "bTextureArray", BuildSettings.bTextureArray, DefaultSettings.bTextureArray);
 	WriteCbFieldWithDefault<bool>(Writer, "bVolume", BuildSettings.bVolume, DefaultSettings.bVolume);
+	if ( BuildSettings.bVolume )
+	{
+		WriteCbField<bool>(Writer, "bVolume_ForceNewDDcKey", true); 
+	}
 	WriteCbFieldWithDefault<bool>(Writer, "bLongLatSource", BuildSettings.bLongLatSource, DefaultSettings.bLongLatSource);
 	WriteCbFieldWithDefault<bool>(Writer, "bSRGB", BuildSettings.bSRGB, DefaultSettings.bSRGB);
 	WriteCbFieldWithDefault(Writer, "SourceEncodingOverride", BuildSettings.SourceEncodingOverride, DefaultSettings.SourceEncodingOverride);
diff --git a/Engine/Source/Runtime/Engine/Private/TextureDerivedDataTask.cpp b/Engine/Source/Runtime/Engine/Private/TextureDerivedDataTask.cpp
--- a/Engine/Source/Runtime/Engine/Private/TextureDerivedDataTask.cpp
+++ b/Engine/Source/Runtime/Engine/Private/TextureDerivedDataTask.cpp
@@ -266,11 +266,11 @@ void FTextureSourceData::GetSourceMips(FTextureSource& Source, IImageWrapperModu
 				{
 					int32 MipSizeX = SourceBlock.SizeX;
 					int32 MipSizeY = SourceBlock.SizeY;
+					int32 MipSizeZ = SourceBlock.NumSlices;
 					for (int32 MipIndex = 0; MipIndex < BlockData.NumMips; ++MipIndex)
 					{
 						FImage* SourceMip = new(BlockData.MipsPerLayer[LayerIndex]) FImage(
-							MipSizeX, MipSizeY,
-							BlockData.NumSlices,
+							MipSizeX, MipSizeY, MipSizeZ,
 							LayerData.ImageFormat,
 							LayerData.GammaSpace
 						);
@@ -285,6 +285,10 @@ void FTextureSourceData::GetSourceMips(FTextureSource& Source, IImageWrapperModu
 
 						MipSizeX = FMath::Max(MipSizeX / 2, 1);
 						MipSizeY = FMath::Max(MipSizeY / 2, 1);
+						if ( Source.IsVolume() )
+						{
+							MipSizeZ = FMath::Max(MipSizeZ / 2, 1);
+						}
 					}
 				}
 			}
@@ -821,6 +825,7 @@ void FTextureCacheDerivedDataWorker::DoWork()
 
 				if (bSucceeded && ActiveBuildSettingsPerLayer.Num() > 0)
 				{
+					// Bandaid code to prevent loading bad data from DDC
 					// Code inspired by the texture compressor module as a hot fix for the bad data that might have been push into the ddc in 4.23 or 4.24 
 					const bool bLongLatCubemap = DerivedData->IsCubemap() && DerivedData->GetNumSlices() == 1;
 					int32 MaximumNumberOfMipMaps = TNumericLimits<int32>::Max();
@@ -830,6 +835,8 @@ void FTextureCacheDerivedDataWorker::DoWork()
 					}
 					else
 					{
+						// this looks wrong, should be floor
+						// that's okay because an over-estimate here doesn't hurt
 						MaximumNumberOfMipMaps = FMath::CeilLogTwo(FMath::Max3(DerivedData->SizeX, DerivedData->SizeY, ActiveBuildSettingsPerLayer[0].bVolume ? DerivedData->GetNumSlices() : 1)) + 1;
 					}
 
@@ -1305,6 +1312,7 @@ public:
 
 	static FTexturePlatformData::FStructuredDerivedDataKey GetKey(const UE::DerivedData::FBuildDefinition& BuildDefinition, const UTexture& Texture, bool bUseCompositeTexture)
 	{
+		// DDC2 Key SerializeForKey is here!
 		FTexturePlatformData::FStructuredDerivedDataKey Key;
 		Key.BuildDefinitionKey = BuildDefinition.GetKey().Hash;
 		Key.SourceGuid = Texture.Source.GetId();
diff --git a/Engine/Source/Runtime/ImageCore/Public/ImageCore.h b/Engine/Source/Runtime/ImageCore/Public/ImageCore.h
--- a/Engine/Source/Runtime/ImageCore/Public/ImageCore.h
+++ b/Engine/Source/Runtime/ImageCore/Public/ImageCore.h
@@ -28,6 +28,8 @@
 /* Types
  *****************************************************************************/
 
+// EGammaSpace is in Color.h
+
 namespace ERawImageFormat
 {
 	/**
