Fixing incorrect printf specifier for non-TCHAR string type (thanks devin!)

diff --git a/Engine/Source/Developer/TextureBuild/Private/TextureBuildFunction.cpp b/Engine/Source/Developer/TextureBuild/Private/TextureBuildFunction.cpp
--- a/Engine/Source/Developer/TextureBuild/Private/TextureBuildFunction.cpp
+++ b/Engine/Source/Developer/TextureBuild/Private/TextureBuildFunction.cpp
@@ -485,7 +485,7 @@ void GenericTextureTilingBuildFunction(UE::DerivedData::FBuildContext& Context,
 
 	UE::TextureBuildUtilities::FTextureBuildMetadata BuildMetadata(FCbObject(Context.FindInput(ANSITEXTVIEW("TextureBuildMetadata"))));
 
-	UE_LOG(LogTextureBuildFunction, Display, TEXT("Tiling %s with %S -> %d source mip(s) with a tail of %d..."), *Context.GetName(), *BuildFunctionName, TextureDescription.NumMips, TextureExtendedData.NumMipsInTail);
+	UE_LOG(LogTextureBuildFunction, Display, TEXT("Tiling %s with %hs -> %d source mip(s) with a tail of %d..."), *Context.GetName(), *BuildFunctionName, TextureDescription.NumMips, TextureExtendedData.NumMipsInTail);
 
 	//
 	// Careful - the linear build might have a different streaming mip count than we output due to mip tail
