Fix bug in DDC2 texture build path caused by using SourceMips after it had been cleared by BuildTexture.

diff --git a/Engine/Source/Developer/TextureBuild/Private/TextureBuildFunction.cpp b/Engine/Source/Developer/TextureBuild/Private/TextureBuildFunction.cpp
--- a/Engine/Source/Developer/TextureBuild/Private/TextureBuildFunction.cpp
+++ b/Engine/Source/Developer/TextureBuild/Private/TextureBuildFunction.cpp
@@ -383,7 +383,13 @@ void FTextureBuildFunction::Build(UE::DerivedData::FBuildContext& Context) const
 		return;
 	}
 
-	UE_LOG(LogTextureBuildFunction, Display, TEXT("Compressing %s -> %d source mip(s) (%dx%d) to %s..."), *Context.GetName(), SourceMips.Num(), SourceMips[0].SizeX, SourceMips[0].SizeY, *BuildSettings.TextureFormatName.ToString());
+	// SourceMips will be cleared by BuildTexture.  Store info from it for use later.
+	const int32 SourceMipsNum = SourceMips.Num();
+	const int32 SourceMipsNumSlices = SourceMips[0].NumSlices;
+	const int32 SourceMip0SizeX = SourceMips[0].SizeX;
+	const int32 SourceMip0SizeY = SourceMips[0].SizeY;
+
+	UE_LOG(LogTextureBuildFunction, Display, TEXT("Compressing %s -> %d source mip(s) (%dx%d) to %s..."), *Context.GetName(), SourceMipsNum, SourceMip0SizeX, SourceMip0SizeY, *BuildSettings.TextureFormatName.ToString());
 
 	ITextureCompressorModule& TextureCompressorModule = FModuleManager::GetModuleChecked<ITextureCompressorModule>(TEXTURE_COMPRESSOR_MODULENAME);
 	
@@ -412,7 +418,7 @@ void FTextureBuildFunction::Build(UE::DerivedData::FBuildContext& Context) const
 
 	{
 		int32 CalculatedMip0SizeX = 0, CalculatedMip0SizeY = 0, CalculatedMip0NumSlices = 0;
-		int32 CalculatedMipCount = TextureCompressorModule.GetMipCountForBuildSettings(SourceMips[0].SizeX, SourceMips[0].SizeY, SourceMips[0].NumSlices, SourceMips.Num(), BuildSettings, CalculatedMip0SizeX, CalculatedMip0SizeY, CalculatedMip0NumSlices);
+		int32 CalculatedMipCount = TextureCompressorModule.GetMipCountForBuildSettings(SourceMip0SizeX, SourceMip0SizeY, SourceMipsNumSlices, SourceMipsNum, BuildSettings, CalculatedMip0SizeX, CalculatedMip0SizeY, CalculatedMip0NumSlices);
 		BuildSettings.GetEncodedTextureDescriptionWithPixelFormat(&TextureDescription, (EPixelFormat)CompressedMips[0].PixelFormat, CalculatedMip0SizeX, CalculatedMip0SizeY, CalculatedMip0NumSlices, CalculatedMipCount);
 	}
 
