TextureBuildFunction : methodology to validate RequiredMemoryEstimate vs. actual mem use using standalone texture build workers and process memory counters

diff --git a/Engine/Source/Developer/DerivedDataCache/Private/DerivedDataBuildJobContext.h b/Engine/Source/Developer/DerivedDataCache/Private/DerivedDataBuildJobContext.h
--- a/Engine/Source/Developer/DerivedDataCache/Private/DerivedDataBuildJobContext.h
+++ b/Engine/Source/Developer/DerivedDataCache/Private/DerivedDataBuildJobContext.h
@@ -43,7 +43,7 @@ public:
 
 	inline ECachePolicy GetCachePolicyMask() const final { return CachePolicyMask; }
 	inline EBuildPolicy GetBuildPolicyMask() const final { return BuildPolicyMask; }
-	inline uint64 GetRequiredMemory() const { return RequiredMemory; }
+	uint64 GetRequiredMemory() const override { return RequiredMemory; }
 	inline bool ShouldCheckDeterministicOutput() const { return bDeterministicOutputCheck; }
 
 	void AddConstant(FUtf8StringView Key, FCbObject&& Value);
diff --git a/Engine/Source/Developer/DerivedDataCache/Public/DerivedDataBuildFunction.h b/Engine/Source/Developer/DerivedDataCache/Public/DerivedDataBuildFunction.h
--- a/Engine/Source/Developer/DerivedDataCache/Public/DerivedDataBuildFunction.h
+++ b/Engine/Source/Developer/DerivedDataCache/Public/DerivedDataBuildFunction.h
@@ -85,6 +85,9 @@ public:
 
 	/** Returns the name by which to identify this build for logging and profiling. */
 	virtual const FSharedString& GetName() const = 0;
+	
+	/** Returns estimate of memory required for this operation. */
+	virtual uint64 GetRequiredMemory() const = 0;
 
 	/** Returns the constant with the matching key, or an object with no fields if not found. */
 	virtual FCbObject FindConstant(FUtf8StringView Key) const = 0;
diff --git a/Engine/Source/Developer/TextureBuild/Private/TextureBuildFunction.cpp b/Engine/Source/Developer/TextureBuild/Private/TextureBuildFunction.cpp
--- a/Engine/Source/Developer/TextureBuild/Private/TextureBuildFunction.cpp
+++ b/Engine/Source/Developer/TextureBuild/Private/TextureBuildFunction.cpp
@@ -24,6 +24,8 @@
 #include "TextureBuildUtilities.h"
 #include "TextureCompressorModule.h"
 #include "TextureFormatManager.h"
+#include "Misc/CommandLine.h"
+#include "HAL/FileManager.h"
 
 DEFINE_LOG_CATEGORY_STATIC(LogTextureBuildFunction, Log, All);
 
@@ -546,11 +548,75 @@ void FTextureBuildFunction::Build(UE::DerivedData::FBuildContext& Context) const
 	const int32 SourceMipsNumSlices = SourceMips[0].NumSlices;
 	const int32 SourceMip0SizeX = SourceMips[0].SizeX;
 	const int32 SourceMip0SizeY = SourceMips[0].SizeY;
-
-	UE_LOG(LogTextureBuildFunction, Display, TEXT("Compressing %s -> %d source mip(s) (%dx%d) to %s..."), *Context.GetName(), SourceMipsNum, SourceMip0SizeX, SourceMip0SizeY, *BuildSettings.TextureFormatName.ToString());
+	bool bHasCompositeSource = AssociatedNormalSourceMips.Num() > 0;
+
+	// @todo Oodle : Context.GetName() is the "build.action" file name, we want the Texture name
+	//		(we want to log *both* not one or the other)
+
+	UE_LOG(LogTextureBuildFunction, Display, TEXT("Compressing [%s] from %dx%d (%d slices, %d mips) to %s...%s%s%s%s%s RequiredMemory=%.3f MB"), 
+		*Context.GetName(), 
+		SourceMip0SizeX, SourceMip0SizeY, SourceMipsNumSlices, SourceMipsNum,
+		*BuildSettings.TextureFormatName.ToString(),
+		bHasCompositeSource ? TEXT(" Composite") : TEXT(""),
+		BuildSettings.bVolume ? TEXT(" Volume") : TEXT(""),
+		BuildSettings.bCubemap ? TEXT(" Cube") : TEXT(""),
+		BuildSettings.bLongLatSource ? TEXT(" LongLat") : TEXT(""),
+		BuildSettings.bTextureArray ? TEXT(" Array") : TEXT(""),
+		Context.GetRequiredMemory()/(1024.0*1024)
+		);
 
 	ITextureCompressorModule& TextureCompressorModule = FModuleManager::GetModuleChecked<ITextureCompressorModule>(TEXTURE_COMPRESSOR_MODULENAME);
 	
+	bool DoMemoryCheck = false;
+	
+#if !(WITH_EDITOR) // is standalone TBW
+	// -tbfmemcheck -ansimalloc
+	if ( FParse::Param(FCommandLine::Get(), TEXT("tbfmemcheck")) )
+	{
+		DoMemoryCheck = true;
+		if ( ! FParse::Param(FCommandLine::Get(), TEXT("ansimalloc")) )
+		{
+			UE_LOG(LogTextureBuildFunction, Display, TEXT("NOTE: Memory use report may be inaccurate; use -ansimalloc."));
+		}
+	}
+#endif
+
+	if ( DoMemoryCheck )
+	{
+		// do an encode of a tiny 4x4 image first, with same settings
+		// this runs through the code once, and allocates some of the globals that are init-on-first-use that will stick around
+
+		TArray<FImage> FakeSourceMips;
+		FakeSourceMips.SetNum(1);
+		FImageCore::ResizeImageAllocDest(SourceMips[0],FakeSourceMips[0],4,4);
+		
+		TArray<FImage> FakeAssociatedNormalSourceMips;
+		if ( bHasCompositeSource )
+		{
+			FakeAssociatedNormalSourceMips = FakeSourceMips;
+		}
+
+		TArray<FCompressedImage2D> FakeCompressedMips;
+		uint32 FakeNumMipsInTail;
+		uint32 FakeExtData;
+		UE::TextureBuildUtilities::FTextureBuildMetadata FakeBuildMetadata;
+
+		TextureCompressorModule.BuildTexture(
+			FakeSourceMips,
+			FakeAssociatedNormalSourceMips,
+			BuildSettings,
+			Context.GetName(),
+			FakeCompressedMips,
+			FakeNumMipsInTail,
+			FakeExtData,
+			&FakeBuildMetadata
+			);
+	}
+	
+	FPlatformMemoryStats MemStatsBefore = FPlatformMemory::GetStats();
+
+	// note: getting Metadata here means ComputeMipChainHash is called, unlike in DDC1 use
+
 	TArray<FCompressedImage2D> CompressedMips;
 	uint32 NumMipsInTail;
 	uint32 ExtData;
@@ -570,7 +636,110 @@ void FTextureBuildFunction::Build(UE::DerivedData::FBuildContext& Context) const
 		return;
 	}
 	check(CompressedMips.Num() > 0);
+	// SourceMips may have been freed by BuildTexture, do not use them any more
+	SourceMips.Reset();
+	
+	uint64 BuildMemAllocated = 0;
+	
+	if ( DoMemoryCheck )
+	{
+		// if DoMemoryCheck is off, you could still do this scope to get BuildMemAllocated
+		//	but it would not be accurate, so it would be misleading, so just don't do it
+
+		FPlatformMemoryStats MemStatsAfter = FPlatformMemory::GetStats();
+
+		if ( MemStatsAfter.PeakUsedVirtual == MemStatsBefore.PeakUsedVirtual )
+		{
+			// peak did not occur during BuildTexture
+			//	(it occurred in startup/init)
+			//	so we do not have a useful reading
+			BuildMemAllocated = 0;
+		}
+		else
+		{
+			// take Peak observed during Build and subtract Pagefile before (not Peak before)
+			BuildMemAllocated = MemStatsAfter.PeakUsedVirtual - MemStatsBefore.UsedVirtual;
+		}
+	}
+
+	// log built info :
+	{
+		int32 CompressedMipCount = CompressedMips.Num();
+
+		int64 CompressedDataSizeTotal = 0;
+		for( const FCompressedImage2D & CompressedMip : CompressedMips )
+		{
+			CompressedDataSizeTotal += CompressedMip.RawData.Num();
+		}
+		
+		const FCompressedImage2D & CompressedImage = CompressedMips[0];
+
+		// log what the TextureFormat built :
+		UE_LOG(LogTextureBuildFunction, Display, TEXT("Built texture: %d Mips PF=%d=%s : %dx%dx%d : CompressedDataSize=%lld , MemAllocated = %.3f MB"),
+			//[%.*s] DebugTexturePathName.Len(),DebugTexturePathName.GetData(),
+			CompressedMipCount, (int)CompressedImage.PixelFormat,
+			GetPixelFormatString((EPixelFormat)CompressedImage.PixelFormat),
+			CompressedImage.SizeX, CompressedImage.SizeY, CompressedImage.NumSlicesWithDepth,
+			CompressedDataSizeTotal,
+			BuildMemAllocated/(1024.0*1024));
+
+		// log csv line
+		
+		FString CSVFilename;
+		if ( FParse::Value(FCommandLine::Get(), TEXT("tbfcsv="),CSVFilename) ||
+			FParse::Param(FCommandLine::Get(), TEXT("tbfcsv")) )
+		{
+			if ( CSVFilename.IsEmpty() || CSVFilename[0] == TEXT('-') )
+			{
+				CSVFilename = TEXT("tbf.csv");
+			}
+
+			TUniquePtr<FArchive> OutputArchive(IFileManager::Get().CreateFileWriter(*CSVFilename, FILEWRITE_Append));
+			if (!OutputArchive.IsValid())
+			{
+				UE_LOG(LogTextureBuildFunction, Display, TEXT("Failed to save CSV file %s"), *CSVFilename);
+			}
+			else
+			{
+				OutputArchive->Logf( TEXT("%s,%d,%d,%d,%d,%lld,%s,%s,%s%s%s%s%s,%d,%d,%d,%lld,%.3f,%.3f"),
+					*Context.GetName(), // @todo : we want texture name and the build.action file name both
+					SourceMip0SizeX, SourceMip0SizeY, SourceMipsNumSlices, SourceMipsNum,
+					(int64)SourceMip0SizeX*SourceMip0SizeY*SourceMipsNumSlices,
+					*BuildSettings.TextureFormatName.ToString(),
+					GetPixelFormatString((EPixelFormat)CompressedImage.PixelFormat),
+					bHasCompositeSource ? TEXT(" Composite") : TEXT(""),
+					BuildSettings.bVolume ? TEXT(" Volume") : TEXT(""),
+					BuildSettings.bCubemap ? TEXT(" Cube") : TEXT(""),
+					BuildSettings.bLongLatSource ? TEXT(" LongLat") : TEXT(""),
+					BuildSettings.bTextureArray ? TEXT(" Array") : TEXT(""),
+			
+					CompressedImage.SizeX, CompressedImage.SizeY, CompressedImage.NumSlicesWithDepth,
+					CompressedDataSizeTotal,
+			
+					Context.GetRequiredMemory()/(1024.0*1024),
+					BuildMemAllocated/(1024.0*1024)
+					);
+
+				OutputArchive->Flush();
+			}
+		}
+	}
 
+	if ( DoMemoryCheck )
+	{
+		// add a little wiggle room due to inaccuracy of measurement
+		//	(eg. malloc free lists can hold this much memory, various statics and global lists)
+		uint64 RequiredMemPadded = Context.GetRequiredMemory() + 1024*1024;
+		
+		if ( BuildMemAllocated > RequiredMemPadded )
+		{
+			UE_LOG(LogTextureBuildFunction, Warning, TEXT("BuildMemAllocated (%lld) > RequiredMemPadded (%lld)"),
+				BuildMemAllocated,RequiredMemPadded);
+		}
+
+		// for testing, get a hard stop if we used more memory than the estimate :
+		//check( BuildMemAllocated <= RequiredMemPadded );
+	}
 
 	FChildBuildData OutputData;
 	{
@@ -590,7 +759,7 @@ void FTextureBuildFunction::Build(UE::DerivedData::FBuildContext& Context) const
 		BuildSettings.GetEncodedTextureDescriptionWithPixelFormat(&TextureDescription, (EPixelFormat)CompressedMips[0].PixelFormat, CalculatedMip0SizeX, CalculatedMip0SizeY, CalculatedMip0NumSlices, CalculatedMipCount);
 		OutputData.TextureDescription = MoveTemp(TextureDescription);
 	}
-
+	
 	
 
 	// ExtendedData is only really useful for textures that have a post build step for tiling,
@@ -614,7 +783,7 @@ void FTextureBuildFunction::Build(UE::DerivedData::FBuildContext& Context) const
 
 		OutputData.TextureExtendedData = MoveTemp(ExtendedData);
 	}
-
+		
 	
 	OutputData.NumStreamingMips = OutputData.TextureDescription.GetNumStreamingMips(&OutputData.TextureExtendedData, OutputData.EngineParameters);
 
diff --git a/Engine/Source/Developer/TextureBuildUtilities/Private/TextureBuildUtilities.cpp b/Engine/Source/Developer/TextureBuildUtilities/Private/TextureBuildUtilities.cpp
--- a/Engine/Source/Developer/TextureBuildUtilities/Private/TextureBuildUtilities.cpp
+++ b/Engine/Source/Developer/TextureBuildUtilities/Private/TextureBuildUtilities.cpp
@@ -680,7 +680,7 @@ TEXTUREBUILDUTILITIES_API int64 GetVirtualTextureRequiredMemoryEstimate(const FT
 	// larger of early resize phase and VT build phase :
 	MemoryEstimate = FMath::Max(ResizingPhaseMemUse,MemoryEstimate);
 
-	MemoryEstimate += 64 * 1024; // overhead room
+	MemoryEstimate += 1024 * 1024; // overhead room
 
 	//UE_LOG(LogTextureBuildUtilities,Display,TEXT("GetBuildRequiredMemoryEstimate VT : %.3f MB"),MemoryEstimate/(1024*1024.f));
 
@@ -896,7 +896,7 @@ TEXTUREBUILDUTILITIES_API int64 GetPhysicalTextureBuildMemoryEstimate(const FTex
 	// mem use is the max of the phases :
 	MemoryEstimate = FMath::Max(MemoryEstimate,ResizingPhaseMemUse);
 
-	MemoryEstimate += 64 * 1024; // overhead room
+	MemoryEstimate += 1024 * 1024; // overhead room
 
 	//UE_LOG(LogTextureBuildUtilities,Display,TEXT("GetBuildRequiredMemoryEstimate non-VT : %.3f MB"),MemoryEstimate/(1024*1024.f));
 
